Merging //UE4/Release-4.11 to Dev-Main (//UE4/Dev-Main) (pre copyright update in revision 2804086)

diff --git a/Engine/Build/InstalledEngineFilters.ini b/Engine/Build/InstalledEngineFilters.ini
--- a/Engine/Build/InstalledEngineFilters.ini
+++ b/Engine/Build/InstalledEngineFilters.ini
@@ -82,7 +82,7 @@
 
 {Win64} /Engine/Extras/3dsMaxScripts/...
 {Win64} /Engine/Extras/VisualStudioDebugging/UE4.natvis
-{Win64} /Engine/Extras/Maya_AnimationRiggingTools/...
+{Win64} /Engine/Extras/Maya_AnimationRiggingTools/MayaTools/...
 {Win64} /Engine/Extras/MayaVelocityGridExporter/...
 {Win64} /Engine/Extras/UnrealVS/...
 {Win64} /Engine/Extras/Redist/en-us/*
@@ -105,6 +105,7 @@
 {Mac} /Engine/Build/BatchFiles/RunUAT.sh
 
 {Mac} /Engine/Extras/LLDBDataFormatters/UE4DataFormatters.py
+{Mac} /Engine/Extras/Maya_AnimationRiggingTools/ArtToolsOSX/...
 
 ; Don't want these folders, even if they're part of ToolsForCompile
 {Mac} -/Engine/Binaries/Win32/...
diff --git a/Engine/Plugins/Developer/CodeLiteSourceCodeAccess/Source/CodeLiteSourceCodeAccess/CodeLiteSourceCodeAccess.Build.cs b/Engine/Plugins/Developer/CodeLiteSourceCodeAccess/Source/CodeLiteSourceCodeAccess/CodeLiteSourceCodeAccess.Build.cs
--- a/Engine/Plugins/Developer/CodeLiteSourceCodeAccess/Source/CodeLiteSourceCodeAccess/CodeLiteSourceCodeAccess.Build.cs
+++ b/Engine/Plugins/Developer/CodeLiteSourceCodeAccess/Source/CodeLiteSourceCodeAccess/CodeLiteSourceCodeAccess.Build.cs
@@ -1,3 +1,5 @@
+// Copyright 1998-2015 Epic Games, Inc. All Rights Reserved.
+
 namespace UnrealBuildTool.Rules
 {
 	public class CodeLiteSourceCodeAccess : ModuleRules
diff --git a/Engine/Plugins/Runtime/LeapMotionController/Source/LeapMotionController/Private/LeapMotionBoneActor.cpp b/Engine/Plugins/Runtime/LeapMotionController/Source/LeapMotionController/Private/LeapMotionBoneActor.cpp
--- a/Engine/Plugins/Runtime/LeapMotionController/Source/LeapMotionController/Private/LeapMotionBoneActor.cpp
+++ b/Engine/Plugins/Runtime/LeapMotionController/Source/LeapMotionController/Private/LeapMotionBoneActor.cpp
@@ -1,4 +1,4 @@
-// Fill out your copyright notice in the Description page of Project Settings.
+// Copyright 1998-2015 Epic Games, Inc. All Rights Reserved.
 
 #include "LeapMotionControllerPrivatePCH.h"
 #include "LeapMotionBoneActor.h"
diff --git a/Engine/Plugins/Runtime/LeapMotionController/Source/LeapMotionController/Private/LeapMotionControllerActor.cpp b/Engine/Plugins/Runtime/LeapMotionController/Source/LeapMotionController/Private/LeapMotionControllerActor.cpp
--- a/Engine/Plugins/Runtime/LeapMotionController/Source/LeapMotionController/Private/LeapMotionControllerActor.cpp
+++ b/Engine/Plugins/Runtime/LeapMotionController/Source/LeapMotionController/Private/LeapMotionControllerActor.cpp
@@ -1,4 +1,4 @@
-// Fill out your copyright notice in the Description page of Project Settings.
+// Copyright 1998-2015 Epic Games, Inc. All Rights Reserved.
 
 #include "LeapMotionControllerPrivatePCH.h"
 #include "LeapMotionControllerActor.h"
diff --git a/Engine/Plugins/Runtime/LeapMotionController/Source/LeapMotionController/Public/Leap_NoPI.h b/Engine/Plugins/Runtime/LeapMotionController/Source/LeapMotionController/Public/Leap_NoPI.h
--- a/Engine/Plugins/Runtime/LeapMotionController/Source/LeapMotionController/Public/Leap_NoPI.h
+++ b/Engine/Plugins/Runtime/LeapMotionController/Source/LeapMotionController/Public/Leap_NoPI.h
@@ -1,3 +1,5 @@
+// Copyright 1998-2015 Epic Games, Inc. All Rights Reserved.
+
 #pragma once
 
 #ifdef PI
diff --git a/Engine/Plugins/Runtime/OculusRift/Source/OculusRift/Private/OculusRiftMeshAssets.h b/Engine/Plugins/Runtime/OculusRift/Source/OculusRift/Private/OculusRiftMeshAssets.h
--- a/Engine/Plugins/Runtime/OculusRift/Source/OculusRift/Private/OculusRiftMeshAssets.h
+++ b/Engine/Plugins/Runtime/OculusRift/Source/OculusRift/Private/OculusRiftMeshAssets.h
@@ -1,4 +1,4 @@
-// Copyright 1998-2015 Epic Gamesf, Inc. All Rights Reserved.
+// Copyright 1998-2015 Epic Games, Inc. All Rights Reserved.
 
 #pragma once
 
diff --git a/Engine/Source/Developer/FunctionalTesting/Private/FunctionalTestingManager.cpp b/Engine/Source/Developer/FunctionalTesting/Private/FunctionalTestingManager.cpp
--- a/Engine/Source/Developer/FunctionalTesting/Private/FunctionalTestingManager.cpp
+++ b/Engine/Source/Developer/FunctionalTesting/Private/FunctionalTestingManager.cpp
@@ -379,7 +379,4 @@ void UFunctionalTestingManager::SetReproString(FString ReproString)
 	{
 		ReproString.ParseIntoArray(TestReproStrings, FFunctionalTesting::ReproStringTestSeparator, /*InCullEmpty=*/true);
 	}
-}
-
-
-#undef LOCTEXT_NAMESPACE
+}
\ No newline at end of file
diff --git a/Engine/Source/Developer/GameplayDebugger/Private/GameplayDebugger.cpp b/Engine/Source/Developer/GameplayDebugger/Private/GameplayDebugger.cpp
--- a/Engine/Source/Developer/GameplayDebugger/Private/GameplayDebugger.cpp
+++ b/Engine/Source/Developer/GameplayDebugger/Private/GameplayDebugger.cpp
@@ -212,8 +212,8 @@ void FGameplayDebugger::CreateSnappingOptionsMenu(FMenuBuilder& Builder)
 	{
 		Builder.AddMenuSeparator();
 		Builder.AddSubMenu(
-			LOCTEXT("Test_GameplayDebugger_Menu", "Gameplay Debugger"),
-			LOCTEXT("Test_GameplayDebugger_Menu_Tooltip", "Quick setting for Gameplay Debugger tool in selected view"),
+			LOCTEXT("Test_GameplayDebugger_SnappingOptions_Menu", "Gameplay Debugger"),
+			LOCTEXT("Test_GameplayDebugger_SnappingOptions_Menu_Tooltip", "Quick setting for Gameplay Debugger tool in selected view"),
 			FNewMenuDelegate::CreateRaw(this, &FGameplayDebugger::CreateSettingSubMenu)
 			);
 	}
diff --git a/Engine/Source/Developer/HTML5/HTML5PlatformEditor/Private/HTML5PlatformEditorModule.cpp b/Engine/Source/Developer/HTML5/HTML5PlatformEditor/Private/HTML5PlatformEditorModule.cpp
--- a/Engine/Source/Developer/HTML5/HTML5PlatformEditor/Private/HTML5PlatformEditorModule.cpp
+++ b/Engine/Source/Developer/HTML5/HTML5PlatformEditor/Private/HTML5PlatformEditorModule.cpp
@@ -44,7 +44,7 @@ class FHTML5PlatformEditorModule
 		{
 			SettingsModule->RegisterSettings("Project", "Platforms", "HTML5",
 				LOCTEXT("TargetSettingsName", "HTML5"),
-				LOCTEXT("SDKSettingsDescription", "Settings for HTML5"),
+				LOCTEXT("TargetSettingsDescription", "Settings for HTML5"),
 				GetMutableDefault<UHTML5TargetSettings>()
 			);
 
diff --git a/Engine/Source/Developer/LauncherServices/Private/Launcher/LauncherWorker.cpp b/Engine/Source/Developer/LauncherServices/Private/Launcher/LauncherWorker.cpp
--- a/Engine/Source/Developer/LauncherServices/Private/Launcher/LauncherWorker.cpp
+++ b/Engine/Source/Developer/LauncherServices/Private/Launcher/LauncherWorker.cpp
@@ -454,7 +454,7 @@ FString FLauncherWorker::CreateUATCommand( const ILauncherProfileRef& InProfile,
 			UATCommand += MapList;
 
 			FCommandDesc Desc;
-			FText Command = LOCTEXT("LauncherCookDesc", "Starting cook on the fly server");
+			FText Command = LOCTEXT("LauncherCookOnTheFlyDesc", "Starting cook on the fly server");
 			Desc.Name = "Cook Server Task";
 			Desc.Desc = Command.ToString();
 			Desc.EndText = TEXT("********** COOK COMMAND COMPLETED **********");
diff --git a/Engine/Source/Developer/LogVisualizer/Private/SVisualLoggerView.cpp b/Engine/Source/Developer/LogVisualizer/Private/SVisualLoggerView.cpp
--- a/Engine/Source/Developer/LogVisualizer/Private/SVisualLoggerView.cpp
+++ b/Engine/Source/Developer/LogVisualizer/Private/SVisualLoggerView.cpp
@@ -174,7 +174,7 @@ void SVisualLoggerView::Construct(const FArguments& InArgs, const TSharedRef<FUI
 								[
 									SAssignNew(SearchBox, SSearchBox)
 									.OnTextChanged(InArgs._OnFiltersSearchChanged)
-									.HintText_Lambda([Settings]()->FText{return Settings->bSearchInsideLogs ? LOCTEXT("FiltersSearchHint", "Log Data Search") : LOCTEXT("FiltersSearchHint", "Log Category Search"); })
+									.HintText_Lambda([Settings]()->FText{return Settings->bSearchInsideLogs ? LOCTEXT("DataFiltersSearchHint", "Log Data Search") : LOCTEXT("CategoryFiltersSearchHint", "Log Category Search"); })
 								]
 #endif
 							]
diff --git a/Engine/Source/Developer/ProjectLauncher/Private/Widgets/Project/SProjectLauncherProjectPicker.cpp b/Engine/Source/Developer/ProjectLauncher/Private/Widgets/Project/SProjectLauncherProjectPicker.cpp
--- a/Engine/Source/Developer/ProjectLauncher/Private/Widgets/Project/SProjectLauncherProjectPicker.cpp
+++ b/Engine/Source/Developer/ProjectLauncher/Private/Widgets/Project/SProjectLauncherProjectPicker.cpp
@@ -132,7 +132,7 @@ FText SProjectLauncherProjectPicker::HandleProjectComboButtonText( ) const
 
 FText SProjectLauncherProjectPicker::HandleProjectComboButtonToolTip( ) const
 {
-	return LOCTEXT("SelectProjectText", "Select or browse for a project");
+	return LOCTEXT("SelectProjectText_Tooltip", "Select or browse for a project");
 }
 
 
diff --git a/Engine/Source/Developer/ProjectLauncher/Private/Widgets/SProjectLauncherProgress.h b/Engine/Source/Developer/ProjectLauncher/Private/Widgets/SProjectLauncherProgress.h
--- a/Engine/Source/Developer/ProjectLauncher/Private/Widgets/SProjectLauncherProgress.h
+++ b/Engine/Source/Developer/ProjectLauncher/Private/Widgets/SProjectLauncherProgress.h
@@ -396,12 +396,12 @@ private:
 
 	FText GetRerunButtonToolTip() const
 	{
-		return LOCTEXT("DoneButtonCloseTooltip", "Run this launch profile.");
+		return LOCTEXT("RerunButtonTooltip", "Run this launch profile.");
 	}
 
 	FText GetRerunButtonText() const
 	{
-		return LOCTEXT("DoneButtonDoneLabel", "Run");
+		return LOCTEXT("RerunButtonLabel", "Run");
 	}
 
 	bool IsDoneButtonEnabled() const
diff --git a/Engine/Source/Developer/SharedSettingsWidgets/Private/SSettingsEditorCheckoutNotice.cpp b/Engine/Source/Developer/SharedSettingsWidgets/Private/SSettingsEditorCheckoutNotice.cpp
--- a/Engine/Source/Developer/SharedSettingsWidgets/Private/SSettingsEditorCheckoutNotice.cpp
+++ b/Engine/Source/Developer/SharedSettingsWidgets/Private/SSettingsEditorCheckoutNotice.cpp
@@ -198,7 +198,7 @@ FText SSettingsEditorCheckoutNotice::HandleLockedStatusText() const
 	
 	return FText::Format(ISourceControlModule::Get().IsEnabled() ?
 		LOCTEXT("DefaultSettingsNotice_WithSourceControl", "These settings are saved in {0}, which is currently NOT checked out.") :
-		LOCTEXT("DefaultSettingsNotice_Source", "These settings are saved in {0}, which is currently NOT writable."), ConfigFilename);
+		LOCTEXT("DefaultSettingsNotice_NotWritable", "These settings are saved in {0}, which is currently NOT writable."), ConfigFilename);
 }
 
 FText SSettingsEditorCheckoutNotice::HandleUnlockedStatusText() const
diff --git a/Engine/Source/Developer/SlateReflector/Private/Widgets/SWidgetReflector.cpp b/Engine/Source/Developer/SlateReflector/Private/Widgets/SWidgetReflector.cpp
--- a/Engine/Source/Developer/SlateReflector/Private/Widgets/SWidgetReflector.cpp
+++ b/Engine/Source/Developer/SlateReflector/Private/Widgets/SWidgetReflector.cpp
@@ -51,7 +51,7 @@ struct FLoggedEvent
 
 	FText ToText()
 	{
-		return FText::Format(NSLOCTEXT("","","{0}  |  {1}"), EventText, HandlerText);
+		return FText::Format(LOCTEXT("LoggedEvent","{0}  |  {1}"), EventText, HandlerText);
 	}
 	
 	FInputEvent Event;
@@ -930,7 +930,7 @@ TSharedRef<SDockTab> SWidgetReflector::SpawnSlateStatsTab(const FSpawnTabArgs& A
 						[
 							SNew(STextBlock)
 							.WrapTextAt(200.0f)
-							.Text( LOCTEXT("StatsSamplingIntervalTooltip", "the interval (in seconds) to integrate stats before updating the averages.") )
+							.Text( LOCTEXT("StatsSamplingIntervalLabelTooltip", "the interval (in seconds) to integrate stats before updating the averages.") )
 						]
 					)
 					.Text(LOCTEXT("StatsSampleWindow", "Sampling Interval: "))
diff --git a/Engine/Source/Developer/StandaloneRenderer/Private/Windows/D3D/SlateD3DRenderer.cpp b/Engine/Source/Developer/StandaloneRenderer/Private/Windows/D3D/SlateD3DRenderer.cpp
--- a/Engine/Source/Developer/StandaloneRenderer/Private/Windows/D3D/SlateD3DRenderer.cpp
+++ b/Engine/Source/Developer/StandaloneRenderer/Private/Windows/D3D/SlateD3DRenderer.cpp
@@ -124,7 +124,7 @@ void FSlateD3DRenderer::CreateDevice()
 
 		if (Hr == DXGI_ERROR_UNSUPPORTED)
 		{
-			FPlatformMisc::MessageBoxExt(EAppMsgType::Ok, *NSLOCTEXT("", "", "There is a problem with your graphics card. Please ensure your card meets the minimum system requirements and that you have the latest drivers installed.").ToString(), *NSLOCTEXT("", "UnsupportedVideoCardErrorTitle", "Unsupported Video Card").ToString());
+			FPlatformMisc::MessageBoxExt(EAppMsgType::Ok, *NSLOCTEXT("SlateD3DRenderer", "ProblemWithGraphicsCard", "There is a problem with your graphics card. Please ensure your card meets the minimum system requirements and that you have the latest drivers installed.").ToString(), *NSLOCTEXT("SlateD3DRenderer", "UnsupportedVideoCardErrorTitle", "Unsupported Video Card").ToString());
 		}
 		
 		checkf(SUCCEEDED(Hr), TEXT("D3D11 Error Result %X"), Hr);
diff --git a/Engine/Source/Developer/iOS/IOSPlatformEditor/Private/IOSTargetSettingsCustomization.cpp b/Engine/Source/Developer/iOS/IOSPlatformEditor/Private/IOSTargetSettingsCustomization.cpp
--- a/Engine/Source/Developer/iOS/IOSPlatformEditor/Private/IOSTargetSettingsCustomization.cpp
+++ b/Engine/Source/Developer/iOS/IOSPlatformEditor/Private/IOSTargetSettingsCustomization.cpp
@@ -1017,7 +1017,7 @@ FReply FIOSTargetSettingsCustomization::OnInstallProvisionClicked()
 
 		bOpened = DesktopPlatform->OpenFileDialog(
 			ParentWindowWindowHandle,
-			LOCTEXT("ImportDialogTitle", "Import Provision").ToString(),
+			LOCTEXT("ImportProvisionDialogTitle", "Import Provision").ToString(),
 			FPaths::GetProjectFilePath(),
 			TEXT(""),
 			FileTypes,
@@ -1103,7 +1103,7 @@ FReply FIOSTargetSettingsCustomization::OnInstallCertificateClicked()
 
 		bOpened = DesktopPlatform->OpenFileDialog(
 			ParentWindowWindowHandle,
-			LOCTEXT("ImportDialogTitle", "Import Certificate").ToString(),
+			LOCTEXT("ImportCertificateDialogTitle", "Import Certificate").ToString(),
 			FPaths::GetProjectFilePath(),
 			TEXT(""),
 			FileTypes,
diff --git a/Engine/Source/Editor/AnimGraph/Classes/AnimGraphNode_AssetPlayerBase.cpp b/Engine/Source/Editor/AnimGraph/Classes/AnimGraphNode_AssetPlayerBase.cpp
--- a/Engine/Source/Editor/AnimGraph/Classes/AnimGraphNode_AssetPlayerBase.cpp
+++ b/Engine/Source/Editor/AnimGraph/Classes/AnimGraphNode_AssetPlayerBase.cpp
@@ -1,3 +1,4 @@
+// Copyright 1998-2015 Epic Games, Inc. All Rights Reserved.
 
 #include "AnimGraphPrivatePCH.h"
 #include "AnimGraphNode_AssetPlayerBase.h"
diff --git a/Engine/Source/Editor/AnimGraph/Private/K2Node_AnimGetter.cpp b/Engine/Source/Editor/AnimGraph/Private/K2Node_AnimGetter.cpp
--- a/Engine/Source/Editor/AnimGraph/Private/K2Node_AnimGetter.cpp
+++ b/Engine/Source/Editor/AnimGraph/Private/K2Node_AnimGetter.cpp
@@ -1,3 +1,4 @@
+// Copyright 1998-2015 Epic Games, Inc. All Rights Reserved.
 
 #include "AnimGraphPrivatePCH.h"
 #include "AnimationGraphSchema.h"
diff --git a/Engine/Source/Editor/MainFrame/MainFrame.Build.cs b/Engine/Source/Editor/MainFrame/MainFrame.Build.cs
--- a/Engine/Source/Editor/MainFrame/MainFrame.Build.cs
+++ b/Engine/Source/Editor/MainFrame/MainFrame.Build.cs
@@ -1,4 +1,4 @@
-﻿﻿// Copyright 1998-2015 Epic Games, Inc. All Rights Reserved.
+﻿// Copyright 1998-2015 Epic Games, Inc. All Rights Reserved.
 
 using UnrealBuildTool;
 
diff --git a/Engine/Source/Editor/PlacementMode/Private/PlacementMode.h b/Engine/Source/Editor/PlacementMode/Private/PlacementMode.h
--- a/Engine/Source/Editor/PlacementMode/Private/PlacementMode.h
+++ b/Engine/Source/Editor/PlacementMode/Private/PlacementMode.h
@@ -3,7 +3,7 @@
 
 #include "IPlacementMode.h"
 
-class FPlacementMode : public FEdMode, public IPlacementMode
+class FPlacementMode : public IPlacementMode
 {
 public:
 	FPlacementMode();
diff --git a/Engine/Source/Editor/PlacementMode/Public/IPlacementMode.h b/Engine/Source/Editor/PlacementMode/Public/IPlacementMode.h
--- a/Engine/Source/Editor/PlacementMode/Public/IPlacementMode.h
+++ b/Engine/Source/Editor/PlacementMode/Public/IPlacementMode.h
@@ -1,7 +1,7 @@
 // Copyright 1998-2015 Epic Games, Inc. All Rights Reserved.
 #pragma once
 
-class IPlacementMode
+class IPlacementMode : public FEdMode
 {
 public:
 
@@ -29,6 +29,6 @@ public:
 	/** Removes a widget which when focused would have not ended the active placing session. */
 	virtual void RemoveValidFocusTargetForPlacement( const TWeakPtr< SWidget >& Widget ) = 0;
 
-	/** @return an array of the objects curently being placed */
+	/** @return an array of the objects currently being placed */
 	virtual const TArray< TWeakObjectPtr<UObject> >& GetCurrentlyPlacingObjects() const = 0;
 };
diff --git a/Engine/Source/Editor/PluginWarden/Private/PluginWardenModule.cpp b/Engine/Source/Editor/PluginWarden/Private/PluginWardenModule.cpp
--- a/Engine/Source/Editor/PluginWarden/Private/PluginWardenModule.cpp
+++ b/Engine/Source/Editor/PluginWarden/Private/PluginWardenModule.cpp
@@ -5,17 +5,11 @@
 #include "PluginWardenModule.h"
 #include "SAuthorizingPlugin.h"
 
-#include "SThrobber.h"
-#include "IPortalApplicationWindow.h"
-#include "IPortalServiceLocator.h"
-#include "Application/IPortalApplicationWindow.h"
-#include "DesktopPlatformModule.h"
-
 IMPLEMENT_MODULE( FPluginWardenModule, PluginWarden );
 
 #define LOCTEXT_NAMESPACE "PluginWarden"
 
-static TSet<FString> AuthorizedPlugins;
+TSet<FString> AuthorizedPlugins;
 
 void FPluginWardenModule::StartupModule()
 {
diff --git a/Engine/Source/Editor/PluginWarden/Private/PluginWardenPrivatePCH.h b/Engine/Source/Editor/PluginWarden/Private/PluginWardenPrivatePCH.h
--- a/Engine/Source/Editor/PluginWarden/Private/PluginWardenPrivatePCH.h
+++ b/Engine/Source/Editor/PluginWarden/Private/PluginWardenPrivatePCH.h
@@ -4,3 +4,6 @@
 
 #include "UnrealEd.h"
 #include "SlateBasics.h"
+
+#include "SThrobber.h"
+#include "DesktopPlatformModule.h"
diff --git a/Engine/Source/Editor/PluginWarden/Private/SAuthorizingPlugin.cpp b/Engine/Source/Editor/PluginWarden/Private/SAuthorizingPlugin.cpp
--- a/Engine/Source/Editor/PluginWarden/Private/SAuthorizingPlugin.cpp
+++ b/Engine/Source/Editor/PluginWarden/Private/SAuthorizingPlugin.cpp
@@ -1,12 +1,11 @@
 // Copyright 1998-2015 Epic Games, Inc. All Rights Reserved.
 
 #include "PluginWardenPrivatePCH.h"
+
 #include "SAuthorizingPlugin.h"
 
 #define LOCTEXT_NAMESPACE "PluginWarden"
 
-extern TSet<FString> AuthorizedPlugins;
-
 void SAuthorizingPlugin::Construct(const FArguments& InArgs, const TSharedRef<SWindow>& InParentWindow, const FText& InPluginFriendlyName, const FString& InPluginGuid, TFunction<void()> InAuthorizedCallback)
 {
 	CurrentState = EPluginAuthorizationState::Initializing;
@@ -51,7 +50,10 @@ void SAuthorizingPlugin::Construct(const FArguments& InArgs, const TSharedRef<SW
 		]
 	];
 
-	PortalWindowService = GEditor->GetServiceLocator()->GetServiceRef<IPortalApplicationWindow>();
+	TSharedRef<IPortalServiceLocator> ServiceLocator = GEditor->GetServiceLocator();
+	PortalWindowService = ServiceLocator->GetServiceRef<IPortalApplicationWindow>();
+	PortalUserService = ServiceLocator->GetServiceRef<IPortalUser>();
+	PortalUserLoginService = ServiceLocator->GetServiceRef<IPortalUserLogin>();
 }
 
 FText SAuthorizingPlugin::GetWaitingText() const
@@ -61,11 +63,17 @@ FText SAuthorizingPlugin::GetWaitingText() const
 	case EPluginAuthorizationState::Initializing:
 	case EPluginAuthorizationState::StartLauncher:
 		return LOCTEXT("StartingLauncher", "Starting Epic Games Launcher...");
-	case EPluginAuthorizationState::WaitingForAvailableLauncher:
+	case EPluginAuthorizationState::StartLauncher_Waiting:
 		return LOCTEXT("ConnectingToLauncher", "Connecting...");
 	case EPluginAuthorizationState::AuthorizePlugin:
-	case EPluginAuthorizationState::WaitingForPluginAuthorization:
+	case EPluginAuthorizationState::AuthorizePlugin_Waiting:
 		return FText::Format(LOCTEXT("CheckingIfYouCanUseFormat", "Checking license for {0}..."), PluginFriendlyName);
+	case EPluginAuthorizationState::IsUserSignedIn:
+	case EPluginAuthorizationState::IsUserSignedIn_Waiting:
+		return LOCTEXT("CheckingIfUserSignedIn", "Authorization failed, checking user information...");
+	case EPluginAuthorizationState::SigninRequired:
+	case EPluginAuthorizationState::SigninRequired_Waiting:
+		return LOCTEXT("NeedUserToLoginToCheck", "Authorization failed, Sign-in required...");
 	}
 
 	return LOCTEXT("Processing", "Processing...");
@@ -101,7 +109,7 @@ EActiveTimerReturnType SAuthorizingPlugin::RefreshStatus(double InCurrentTime, f
 			{
 				if ( DesktopPlatform->OpenLauncher(false, FString(), FString(" -silent")) )
 				{
-					CurrentState = EPluginAuthorizationState::WaitingForAvailableLauncher;
+					CurrentState = EPluginAuthorizationState::StartLauncher_Waiting;
 				}
 				else
 				{
@@ -114,9 +122,9 @@ EActiveTimerReturnType SAuthorizingPlugin::RefreshStatus(double InCurrentTime, f
 			}
 			break;
 		}
-		case EPluginAuthorizationState::WaitingForAvailableLauncher:
+		case EPluginAuthorizationState::StartLauncher_Waiting:
 		{
-			if ( PortalWindowService->IsAvailable() )
+			if ( PortalWindowService->IsAvailable() && PortalUserService->IsAvailable() )
 			{
 				CurrentState = EPluginAuthorizationState::AuthorizePlugin;
 			}
@@ -129,31 +137,90 @@ EActiveTimerReturnType SAuthorizingPlugin::RefreshStatus(double InCurrentTime, f
 		case EPluginAuthorizationState::AuthorizePlugin:
 		{
 			WaitingTime = 0;
-			// TODO GET API TO CHECK PLUGIN
-			CurrentState = EPluginAuthorizationState::WaitingForPluginAuthorization;
+			EntitlementResult = PortalUserService->IsEntitledToItem(PluginGuid, EEntitlementCacheLevelRequest::Memory);
+			CurrentState = EPluginAuthorizationState::AuthorizePlugin_Waiting;
 			break;
 		}
-		case EPluginAuthorizationState::WaitingForPluginAuthorization:
+		case EPluginAuthorizationState::AuthorizePlugin_Waiting:
 		{
 			WaitingTime += InDeltaTime;
-			// TODO GET API TO CHECK PLUGIN
+			
+			check(EntitlementResult.GetFuture().IsValid());
+			if ( EntitlementResult.GetFuture().IsReady() )
+			{
+				FPortalUserIsEntitledToItemResult Entitlement = EntitlementResult.GetFuture().Get();
+				if ( Entitlement.IsEntitled )
+				{
+					CurrentState = EPluginAuthorizationState::Authorized;
+				}
+				else
+				{
+					CurrentState = EPluginAuthorizationState::IsUserSignedIn;
+				}
+			}
+
 			break;
 		}
-		case EPluginAuthorizationState::SigninRequired:
+		case EPluginAuthorizationState::IsUserSignedIn:
+		{
+			WaitingTime = 0;
+			UserDetailsResult = PortalUserService->GetUserDetails();
+			CurrentState = EPluginAuthorizationState::IsUserSignedIn_Waiting;
+			break;
+		}
+		case EPluginAuthorizationState::IsUserSignedIn_Waiting:
 		{
-			//TODO Move Launcher sign-in to the foreground
+			WaitingTime += InDeltaTime;
 
-			FText SigninMessage = FText::Format(LOCTEXT("UnauthorizedNeedSignin", "We can't tell if you've purchased {0}.\n\nCan you please sign-in to Epic Launcher and click [OK] when complete?"), PluginFriendlyName);
-			EAppReturnType::Type Response = FMessageDialog::Open(EAppMsgType::OkCancel, SigninMessage);
-			if ( Response == EAppReturnType::Cancel )
+			check(UserDetailsResult.GetFuture().IsValid());
+			if ( UserDetailsResult.GetFuture().IsReady() )
 			{
-				CurrentState = EPluginAuthorizationState::Canceled;
+				FPortalUserDetails UserDetails = UserDetailsResult.GetFuture().Get();
+
+				if ( UserDetails.IsSignedIn )
+				{
+					// if the user is signed in, and we're at this stage, we know they are unauthorized.
+					CurrentState = EPluginAuthorizationState::Unauthorized;
+				}
+				else
+				{
+					// If they're not signed in, but they were unauthorized, they may have purchased it
+					// they may just need to sign-in.
+					if ( PortalUserLoginService->IsAvailable() )
+					{
+						CurrentState = EPluginAuthorizationState::SigninRequired;
+					}
+				}
 			}
-			else
+
+			break;
+		}
+		case EPluginAuthorizationState::SigninRequired:
+		{
+			WaitingTime = 0;
+			UserSigninResult = PortalUserLoginService->PromptUserForSignIn();
+			CurrentState = EPluginAuthorizationState::SigninRequired_Waiting;
+			break;
+		}
+		case EPluginAuthorizationState::SigninRequired_Waiting:
+		{
+			// We don't advance the wait time in the sign-required state, as this may take a long time.
+
+			check(UserSigninResult.GetFuture().IsValid());
+			if ( UserSigninResult.GetFuture().IsReady() )
 			{
-				// Return to the authorize plug-in state now that the user has signed in.
-				CurrentState = EPluginAuthorizationState::AuthorizePlugin;
+				bool IsUserSignedIn = UserSigninResult.GetFuture().Get();
+				if ( IsUserSignedIn )
+				{
+					// if the user is signed in, and we're at this stage, we know they are unauthorized.
+					CurrentState = EPluginAuthorizationState::AuthorizePlugin;
+				}
+				else
+				{
+					CurrentState = EPluginAuthorizationState::Unauthorized;
+				}
 			}
+
 			break;
 		}
 		case EPluginAuthorizationState::Authorized:
@@ -182,8 +249,10 @@ EActiveTimerReturnType SAuthorizingPlugin::RefreshStatus(double InCurrentTime, f
 		// If we're in a waiting state, check to see if we're over the timeout period.
 		switch ( CurrentState )
 		{
-		case EPluginAuthorizationState::WaitingForAvailableLauncher:
-		case EPluginAuthorizationState::WaitingForPluginAuthorization:
+		case EPluginAuthorizationState::StartLauncher_Waiting:
+		case EPluginAuthorizationState::AuthorizePlugin_Waiting:
+		case EPluginAuthorizationState::IsUserSignedIn_Waiting:
+		// We Ignore EPluginAuthorizationState::SigninRequired_Waiting, that state could take forever, the user needs to sign-in or close the dialog.
 		{
 			const float TimeoutSeconds = 15;
 			if ( WaitingTime > TimeoutSeconds )
diff --git a/Engine/Source/Editor/PluginWarden/Private/SAuthorizingPlugin.h b/Engine/Source/Editor/PluginWarden/Private/SAuthorizingPlugin.h
--- a/Engine/Source/Editor/PluginWarden/Private/SAuthorizingPlugin.h
+++ b/Engine/Source/Editor/PluginWarden/Private/SAuthorizingPlugin.h
@@ -1,23 +1,34 @@
 // Copyright 1998-2015 Epic Games, Inc. All Rights Reserved.
 #pragma once
 
+#include "IPortalServiceLocator.h"
+#include "Account/IPortalUser.h"
+#include "Account/IPortalUserLogin.h"
+#include "Application/IPortalApplicationWindow.h"
+
+extern TSet<FString> AuthorizedPlugins;
+
 enum class EPluginAuthorizationState
 {
 	Initializing,
 	StartLauncher,
-	WaitingForAvailableLauncher,
+	StartLauncher_Waiting,
 	AuthorizePlugin,
-	WaitingForPluginAuthorization,
+	AuthorizePlugin_Waiting,
+	IsUserSignedIn,
+	IsUserSignedIn_Waiting,
+	SigninRequired,
+	SigninRequired_Waiting,
 	Authorized,
 	Unauthorized,
-	SigninRequired,
 	LauncherStartFailed,
 	Timeout,
 	Canceled,
 };
 
-class IPortalApplicationWindow;
-
+/**
+ * The authorizing plug-in ui guides the user through the process of certifying their access to the plug-in.
+ */
 class SAuthorizingPlugin : public SCompoundWidget
 {
 public:
@@ -27,6 +38,8 @@ public:
 
 	void Construct(const FArguments& InArgs, const TSharedRef<SWindow>& InParentWindow, const FText& InPluginFriendlyName, const FString& InPluginGuid, TFunction<void()> InAuthorizedCallback);
 
+private:
+
 	FText GetWaitingText() const;
 
 	EActiveTimerReturnType RefreshStatus(double InCurrentTime, float InDeltaTime);
@@ -52,12 +65,27 @@ private:
 	/** The amount of time we've been waiting for confirmation for a given step.  It's possible a problem may arise and we need to timeout. */
 	float WaitingTime;
 
-	/** A pointer to the portal application communication service. */
+	/** The portal application communication service. */
 	TSharedPtr<IPortalApplicationWindow> PortalWindowService;
 
+	/** The portal user service, to allow us to check entitlements for plugins. */
+	TSharedPtr<IPortalUser> PortalUserService;
+
+	/** The portal user login service, to allow us to trigger a prompt to sign-in, if required. */
+	TSharedPtr<IPortalUserLogin> PortalUserLoginService;
+
 	/** The current state of the plug-in auth pipeline. */
 	EPluginAuthorizationState CurrentState;
 
+	/** The entitlement result we may be waiting on. */
+	TAsyncResult<FPortalUserIsEntitledToItemResult> EntitlementResult;
+
+	/** The result from the request for user details. */
+	TAsyncResult<FPortalUserDetails> UserDetailsResult;
+
+	/** The result from requesting the user sign-in, they may sign-in, they may cancel. */
+	TAsyncResult<bool> UserSigninResult;
+
 	/** If the user is authorized to us the plug-in, we'll call this function to alert the plug-in that everything is good to go. */
 	TFunction<void()> AuthorizedCallback;
 };
diff --git a/Engine/Source/Editor/PropertyEditor/Private/SDetailsView.cpp b/Engine/Source/Editor/PropertyEditor/Private/SDetailsView.cpp
--- a/Engine/Source/Editor/PropertyEditor/Private/SDetailsView.cpp
+++ b/Engine/Source/Editor/PropertyEditor/Private/SDetailsView.cpp
@@ -288,7 +288,7 @@ void SDetailsView::ForceRefresh()
 		}
 	}
 
-	SetObjectArrayPrivate( NewObjectList );
+	BeginSelectingObjects(NewObjectList);
 }
 
 
@@ -305,7 +305,7 @@ void SDetailsView::SetObjects( const TArray<UObject*>& InObjects, bool bForceRef
 
 		if( bForceRefresh || ShouldSetNewObjects( ObjectWeakPtrs ) )
 		{
-			SetObjectArrayPrivate( ObjectWeakPtrs );
+			BeginSelectingObjects(ObjectWeakPtrs);
 		}
 	}
 }
@@ -316,7 +316,7 @@ void SDetailsView::SetObjects( const TArray< TWeakObjectPtr< UObject > >& InObje
 	{
 		if( bForceRefresh || ShouldSetNewObjects( InObjects ) )
 		{
-			SetObjectArrayPrivate( InObjects );
+			BeginSelectingObjects(InObjects);
 		}
 	}
 }
@@ -350,7 +350,7 @@ void SDetailsView::RemoveInvalidObjects()
 
 	if (!bAllFound)
 	{
-		SetObjectArrayPrivate(ResetArray);
+		BeginSelectingObjects(ResetArray);
 	}
 }
 
@@ -400,7 +400,14 @@ bool SDetailsView::ShouldSetNewObjects( const TArray< TWeakObjectPtr< UObject >
 	return bShouldSetObjects;
 }
 
-void SDetailsView::SetObjectArrayPrivate( const TArray< TWeakObjectPtr< UObject > >& InObjects )
+void SDetailsView::BeginSelectingObjects(const TArray< TWeakObjectPtr< UObject > >& InObjects)
+{
+	RegisterActiveTimer(0.f, FWidgetActiveTimerDelegate::CreateLambda([=] (double, float) {
+		return FinishSelectingObjects(InObjects);
+	}));
+}
+
+EActiveTimerReturnType SDetailsView::FinishSelectingObjects(const TArray< TWeakObjectPtr< UObject > > InObjects)
 {
 	double StartTime = FPlatformTime::Seconds();
 
@@ -495,6 +502,8 @@ void SDetailsView::SetObjectArrayPrivate( const TArray< TWeakObjectPtr< UObject
 	OnObjectArrayChanged.ExecuteIfBound(Title, InObjects);
 
 	double ElapsedTime = FPlatformTime::Seconds() - StartTime;
+
+	return EActiveTimerReturnType::Stop;
 }
 
 void SDetailsView::ReplaceObjects( const TMap<UObject*, UObject*>& OldToNewObjectMap )
@@ -533,9 +542,8 @@ void SDetailsView::ReplaceObjects( const TMap<UObject*, UObject*>& OldToNewObjec
 
 	if( bObjectsReplaced )
 	{
-		SetObjectArrayPrivate( NewObjectList );
+		BeginSelectingObjects(NewObjectList);
 	}
-
 }
 
 void SDetailsView::RemoveDeletedObjects( const TArray<UObject*>& DeletedObjects )
@@ -561,7 +569,7 @@ void SDetailsView::RemoveDeletedObjects( const TArray<UObject*>& DeletedObjects
 	// if any objects were replaced update the observed objects
 	if( bObjectsRemoved )
 	{
-		SetObjectArrayPrivate( NewObjectList );
+		BeginSelectingObjects(NewObjectList);
 	}
 }
 
diff --git a/Engine/Source/Editor/PropertyEditor/Private/SDetailsView.h b/Engine/Source/Editor/PropertyEditor/Private/SDetailsView.h
--- a/Engine/Source/Editor/PropertyEditor/Private/SDetailsView.h
+++ b/Engine/Source/Editor/PropertyEditor/Private/SDetailsView.h
@@ -100,10 +100,13 @@ public:
 	{ 
 		return false; 
 	}
+
 private:
 	void RegisterInstancedCustomPropertyLayout( UClass* Class, FOnGetDetailCustomizationInstance DetailLayoutDelegate ) override;
 	void UnregisterInstancedCustomPropertyLayout( UClass* Class ) override;
-	void SetObjectArrayPrivate( const TArray< TWeakObjectPtr< UObject > >& InObjects );
+	
+	void BeginSelectingObjects(const TArray< TWeakObjectPtr< UObject > >& InObjects);
+	EActiveTimerReturnType FinishSelectingObjects(const TArray< TWeakObjectPtr< UObject > > InObjects);
 
 	TSharedRef<SDetailTree> ConstructTreeView( TSharedRef<SScrollBar>& ScrollBar );
 
diff --git a/Engine/Source/Editor/UnrealEd/Private/Commandlets/PackageUtilities.cpp b/Engine/Source/Editor/UnrealEd/Private/Commandlets/PackageUtilities.cpp
--- a/Engine/Source/Editor/UnrealEd/Private/Commandlets/PackageUtilities.cpp
+++ b/Engine/Source/Editor/UnrealEd/Private/Commandlets/PackageUtilities.cpp
@@ -104,7 +104,7 @@ bool NormalizePackageNames( TArray<FString> PackageNames, TArray<FString>& Packa
 
 	if ( PackagePathNames.Num() == 0 )
 	{
-		UE_LOG(LogPackageUtilities, Warning, TEXT("No packages found using '%s'!"), *PackageWildcard);
+		UE_LOG(LogPackageUtilities, Log, TEXT("No packages found using '%s'!"), *PackageWildcard);
 		return false;
 	}
 
diff --git a/Engine/Source/Editor/UnrealEd/Public/EditorModeManager.h b/Engine/Source/Editor/UnrealEd/Public/EditorModeManager.h
--- a/Engine/Source/Editor/UnrealEd/Public/EditorModeManager.h
+++ b/Engine/Source/Editor/UnrealEd/Public/EditorModeManager.h
@@ -94,13 +94,13 @@ public:
 	template <typename SpecificModeType>
 	SpecificModeType* GetActiveModeTyped( FEditorModeID InID )
 	{
-		return (SpecificModeType*)GetActiveMode(InID);
+		return static_cast<SpecificModeType*>(GetActiveMode(InID));
 	}
 
 	template <typename SpecificModeType>
 	const SpecificModeType* GetActiveModeTyped( FEditorModeID InID ) const
 	{
-		return (SpecificModeType*)GetActiveMode(InID);
+		return static_cast<SpecificModeType*>(GetActiveMode(InID));
 	}
 
 	/**
diff --git a/Engine/Source/Editor/UnrealEd/UnrealEd.Build.cs b/Engine/Source/Editor/UnrealEd/UnrealEd.Build.cs
--- a/Engine/Source/Editor/UnrealEd/UnrealEd.Build.cs
+++ b/Engine/Source/Editor/UnrealEd/UnrealEd.Build.cs
@@ -145,7 +145,8 @@ public class UnrealEd : ModuleRules
 				"AddContentDialog",
 				"GameProjectGeneration",
                 "HierarchicalLODUtilities",
-                "Analytics"
+                "Analytics",
+                "PluginWarden"
 			}
 		);
 
diff --git a/Engine/Source/Programs/AutomationTool/GUBP/BuildGraph.cs b/Engine/Source/Programs/AutomationTool/GUBP/BuildGraph.cs
--- a/Engine/Source/Programs/AutomationTool/GUBP/BuildGraph.cs
+++ b/Engine/Source/Programs/AutomationTool/GUBP/BuildGraph.cs
@@ -1,4 +1,6 @@
-﻿using System;
+﻿// Copyright 1998-2015 Epic Games, Inc. All Rights Reserved.
+
+using System;
 using System.Collections.Generic;
 using System.Diagnostics;
 using System.Linq;
@@ -35,6 +37,7 @@ namespace AutomationTool
 		{
 			LinkGraph(AggregateNodeDefinitions, BuildNodeDefinitions, out AggregateNodes, out BuildNodes);
 			FindControllingTriggers(BuildNodes);
+			AddTriggerDependencies(BuildNodes);
 		}
 
 		/// <summary>
@@ -292,5 +295,21 @@ namespace AutomationTool
 				}
 			}
 		}
+
+		/// <summary>
+		/// Adds all the dependencies from nodes behind triggers to the nodes in front of a trigger as pseudo-dependencies of the trigger itself.
+		/// This prevents situations where the trigger can be activated before all the dependencies behind it are complete.
+		/// </summary>
+		void AddTriggerDependencies(List<BuildNode> Nodes)
+		{
+			foreach(BuildNode Node in Nodes)
+			{
+				foreach(TriggerNode TriggerNode in Node.ControllingTriggers)
+				{
+					TriggerNode.OrderDependencies.UnionWith(Node.InputDependencies.Where(x => x != TriggerNode && !x.ControllingTriggers.Contains(TriggerNode)));
+					TriggerNode.OrderDependencies.UnionWith(Node.OrderDependencies.Where(x => x != TriggerNode && !x.ControllingTriggers.Contains(TriggerNode)));
+				}
+			}
+		}
 	}
 }
diff --git a/Engine/Source/Programs/AutomationTool/GUBP/BuildGraphDefinition.cs b/Engine/Source/Programs/AutomationTool/GUBP/BuildGraphDefinition.cs
--- a/Engine/Source/Programs/AutomationTool/GUBP/BuildGraphDefinition.cs
+++ b/Engine/Source/Programs/AutomationTool/GUBP/BuildGraphDefinition.cs
@@ -1,4 +1,6 @@
-﻿using System;
+﻿// Copyright 1998-2015 Epic Games, Inc. All Rights Reserved.
+
+using System;
 using System.Collections.Generic;
 using System.IO;
 using System.Linq;
diff --git a/Engine/Source/Programs/AutomationTool/GUBP/ElectricCommander.cs b/Engine/Source/Programs/AutomationTool/GUBP/ElectricCommander.cs
--- a/Engine/Source/Programs/AutomationTool/GUBP/ElectricCommander.cs
+++ b/Engine/Source/Programs/AutomationTool/GUBP/ElectricCommander.cs
@@ -1,4 +1,6 @@
-﻿using System;
+﻿// Copyright 1998-2015 Epic Games, Inc. All Rights Reserved.
+
+using System;
 using System.Collections.Generic;
 using System.Diagnostics;
 using System.IO;
@@ -292,11 +294,11 @@ namespace AutomationTool
 
 			List<string> ECProps = new List<string>();
 			ECProps.Add(String.Format("TimeIndex={0}", TimeIndex));
-			foreach (BuildNode Node in SortedNodes)
+			foreach (BuildNode Node in SortedNodes.Where(x => x.FrequencyShift != BuildNode.ExplicitFrequencyShift))
 			{
 				ECProps.Add(string.Format("AllNodes/{0}={1}", Node.Name, GetNodeForAllNodesProperty(Node, TimeQuantum)));
 			}
-			foreach (KeyValuePair<BuildNode, int> NodePair in FullNodeListSortKey.Where(x => x.Key.ControllingTriggers.Length == 0))
+			foreach (KeyValuePair<BuildNode, int> NodePair in FullNodeListSortKey.Where(x => x.Key.ControllingTriggers.Length == 0 && x.Key.FrequencyShift != BuildNode.ExplicitFrequencyShift))
 			{
 				ECProps.Add(string.Format("SortKey/{0}={1}", NodePair.Key.Name, NodePair.Value));
 			}
diff --git a/Engine/Source/Programs/AutomationTool/GUBP/JobInfo.cs b/Engine/Source/Programs/AutomationTool/GUBP/JobInfo.cs
--- a/Engine/Source/Programs/AutomationTool/GUBP/JobInfo.cs
+++ b/Engine/Source/Programs/AutomationTool/GUBP/JobInfo.cs
@@ -1,4 +1,6 @@
-﻿using System;
+﻿// Copyright 1998-2015 Epic Games, Inc. All Rights Reserved.
+
+using System;
 using System.Collections.Generic;
 using System.Linq;
 using System.Text;
diff --git a/Engine/Source/Programs/AutomationTool/GUBP/LegacyBranchSetup.cs b/Engine/Source/Programs/AutomationTool/GUBP/LegacyBranchSetup.cs
--- a/Engine/Source/Programs/AutomationTool/GUBP/LegacyBranchSetup.cs
+++ b/Engine/Source/Programs/AutomationTool/GUBP/LegacyBranchSetup.cs
@@ -1,4 +1,5 @@
 ﻿// Copyright 1998-2015 Epic Games, Inc. All Rights Reserved.
+
 using System;
 using System.Collections.Generic;
 using System.Text;
@@ -147,6 +148,10 @@ partial class GUBP
             public List<string> PromotablesWithoutTools = new List<string>();
             public List<string> NodesToRemovePseudoDependencies = new List<string>();
             public List<string> EnhanceAgentRequirements = new List<string>();
+			/// <summary>
+			/// Whether to put game builds behind triggers, so that engine branches don't always build them by default.
+			/// </summary>
+			public bool bGameBuildsBehindTriggers = false;
             /// <summary>
             /// This turns off ALL automatedtesting in ALL branches by default.
             /// A branch hacker must set this value to false explicitly to use automated testing.
diff --git a/Engine/Source/Programs/AutomationTool/GUBP/LegacyNodes.cs b/Engine/Source/Programs/AutomationTool/GUBP/LegacyNodes.cs
--- a/Engine/Source/Programs/AutomationTool/GUBP/LegacyNodes.cs
+++ b/Engine/Source/Programs/AutomationTool/GUBP/LegacyNodes.cs
@@ -1,4 +1,5 @@
 ﻿// Copyright 1998-2015 Epic Games, Inc. All Rights Reserved.
+
 using System;
 using System.Collections.Generic;
 using System.Text;
@@ -1564,7 +1565,7 @@ partial class GUBP
 
     public class CookNode : HostPlatformNode
     {
-        BranchInfo.BranchUProject GameProj;
+        public BranchInfo.BranchUProject GameProj;
         UnrealTargetPlatform TargetPlatform;
         string CookPlatform;
         bool bIsMassive;
diff --git a/Engine/Source/Programs/AutomationTool/GUBP/Tasks/ScriptTask.cs b/Engine/Source/Programs/AutomationTool/GUBP/Tasks/ScriptTask.cs
--- a/Engine/Source/Programs/AutomationTool/GUBP/Tasks/ScriptTask.cs
+++ b/Engine/Source/Programs/AutomationTool/GUBP/Tasks/ScriptTask.cs
@@ -1,4 +1,6 @@
-﻿using Microsoft.CSharp;
+﻿// Copyright 1998-2015 Epic Games, Inc. All Rights Reserved.
+
+using Microsoft.CSharp;
 using System;
 using System.CodeDom.Compiler;
 using System.Collections.Generic;
diff --git a/Engine/Source/Programs/AutomationTool/Scripts/ListMobileDevices.Automation.cs b/Engine/Source/Programs/AutomationTool/Scripts/ListMobileDevices.Automation.cs
--- a/Engine/Source/Programs/AutomationTool/Scripts/ListMobileDevices.Automation.cs
+++ b/Engine/Source/Programs/AutomationTool/Scripts/ListMobileDevices.Automation.cs
@@ -1,4 +1,6 @@
-﻿using System.Collections.Generic;
+﻿// Copyright 1998-2015 Epic Games, Inc. All Rights Reserved.
+
+using System.Collections.Generic;
 using AutomationTool;
 using UnrealBuildTool;
 
diff --git a/Engine/Source/Programs/UnrealGameSync/UnrealGameSync/ChildProcess.cs b/Engine/Source/Programs/UnrealGameSync/UnrealGameSync/ChildProcess.cs
--- a/Engine/Source/Programs/UnrealGameSync/UnrealGameSync/ChildProcess.cs
+++ b/Engine/Source/Programs/UnrealGameSync/UnrealGameSync/ChildProcess.cs
@@ -1,4 +1,6 @@
-﻿using Microsoft.Win32.SafeHandles;
+﻿// Copyright 1998-2015 Epic Games, Inc. All Rights Reserved.
+
+using Microsoft.Win32.SafeHandles;
 using System;
 using System.Collections.Generic;
 using System.ComponentModel;
diff --git a/Engine/Source/Programs/UnrealGameSync/UnrealGameSync/Utility.cs b/Engine/Source/Programs/UnrealGameSync/UnrealGameSync/Utility.cs
--- a/Engine/Source/Programs/UnrealGameSync/UnrealGameSync/Utility.cs
+++ b/Engine/Source/Programs/UnrealGameSync/UnrealGameSync/Utility.cs
@@ -1,4 +1,6 @@
-﻿#define USE_NEW_PROCESS_JOBS
+﻿// Copyright 1998-2015 Epic Games, Inc. All Rights Reserved.
+
+#define USE_NEW_PROCESS_JOBS
 
 using System;
 using System.Collections.Generic;
diff --git a/Engine/Source/Programs/UnrealHeaderTool/Private/UHTMakefile/FieldArchiveProxy.h b/Engine/Source/Programs/UnrealHeaderTool/Private/UHTMakefile/FieldArchiveProxy.h
--- a/Engine/Source/Programs/UnrealHeaderTool/Private/UHTMakefile/FieldArchiveProxy.h
+++ b/Engine/Source/Programs/UnrealHeaderTool/Private/UHTMakefile/FieldArchiveProxy.h
@@ -1,5 +1,4 @@
 // Copyright 1998-2015 Epic Games, Inc. All Rights Reserved.
-// Copyright 1998-2015 Epic Games, Inc. All Rights Reserved.
 
 #pragma once
 #include "UHTMakefile/ObjectBaseArchiveProxy.h"
diff --git a/Engine/Source/Runtime/AIModule/Classes/EnvironmentQuery/EnvQueryInstanceBlueprintWrapper.h b/Engine/Source/Runtime/AIModule/Classes/EnvironmentQuery/EnvQueryInstanceBlueprintWrapper.h
--- a/Engine/Source/Runtime/AIModule/Classes/EnvironmentQuery/EnvQueryInstanceBlueprintWrapper.h
+++ b/Engine/Source/Runtime/AIModule/Classes/EnvironmentQuery/EnvQueryInstanceBlueprintWrapper.h
@@ -6,6 +6,8 @@
 #include "EnvironmentQuery/EQSQueryResultSourceInterface.h"
 #include "EnvQueryInstanceBlueprintWrapper.generated.h"
 
+struct FEnvQueryRequest;
+
 UCLASS(BlueprintType, meta = (DisplayName = "EQS Query Instance"))
 class AIMODULE_API UEnvQueryInstanceBlueprintWrapper : public UObject, public IEQSQueryResultSourceInterface
 {
@@ -34,11 +36,8 @@ protected:
 
 public:
 	UEnvQueryInstanceBlueprintWrapper(const FObjectInitializer& ObjectInitializer);
-
-	void OnQueryFinished(TSharedPtr<FEnvQueryResult> Result);
-
-	void SetRunMode(EEnvQueryRunMode::Type InRunMode) { RunMode = InRunMode; }
-	void SetQueryID(int32 InQueryID) { QueryID = InQueryID; }
+	
+	void RunQuery(const EEnvQueryRunMode::Type InRunMode, FEnvQueryRequest& QueryRequest);
 
 	UFUNCTION(BlueprintPure, Category = "AI|EQS")
 	float GetItemScore(int32 ItemIndex);
@@ -50,4 +49,7 @@ public:
 	/** returns an array of locations generated by the query. If the query generated Actors the the array is filled with their locations */
 	UFUNCTION(BlueprintPure, Category = "AI|EQS")
 	TArray<FVector> GetResultsAsLocations();
+
+protected:
+	void OnQueryFinished(TSharedPtr<FEnvQueryResult> Result);
 };
\ No newline at end of file
diff --git a/Engine/Source/Runtime/AIModule/Classes/EnvironmentQuery/EnvQueryManager.h b/Engine/Source/Runtime/AIModule/Classes/EnvironmentQuery/EnvQueryManager.h
--- a/Engine/Source/Runtime/AIModule/Classes/EnvironmentQuery/EnvQueryManager.h
+++ b/Engine/Source/Runtime/AIModule/Classes/EnvironmentQuery/EnvQueryManager.h
@@ -14,6 +14,7 @@ class UEnvQueryOption;
 struct FEnvQueryInstance;
 struct FEnvNamedValue;
 class UEnvQueryTest;
+class UEnvQueryInstanceBlueprintWrapper;
 
 /** wrapper for easy query execution */
 USTRUCT()
@@ -178,6 +179,9 @@ class AIMODULE_API UEnvQueryManager : public UObject, public FTickableGameObject
 	UFUNCTION(BlueprintCallable, Category = "AI|EQS", meta = (WorldContext = "WorldContext", AdvancedDisplay = "WrapperClass"))
 	static UEnvQueryInstanceBlueprintWrapper* RunEQSQuery(UObject* WorldContext, UEnvQuery* QueryTemplate, UObject* Querier, TEnumAsByte<EEnvQueryRunMode::Type> RunMode, TSubclassOf<UEnvQueryInstanceBlueprintWrapper> WrapperClass);
 
+	void RegisterActiveWrapper(UEnvQueryInstanceBlueprintWrapper& Wrapper);
+	void UnregisterActiveWrapper(UEnvQueryInstanceBlueprintWrapper& Wrapper);
+
 	static void SetAllowTimeSlicing(bool bAllowTimeSlicing);
 
 #if USE_EQS_DEBUGGER
@@ -205,6 +209,9 @@ protected:
 	UPROPERTY(transient)
 	TArray<UEnvQueryContext*> LocalContexts;
 
+	UPROPERTY()
+	TArray<UEnvQueryInstanceBlueprintWrapper*> GCShieldedWrappers;
+
 	/** local contexts mapped by class names */
 	TMap<FName, UEnvQueryContext*> LocalContextMap;
 
diff --git a/Engine/Source/Runtime/AIModule/Private/EnvironmentQuery/EnvQueryInstanceBlueprintWrapper.cpp b/Engine/Source/Runtime/AIModule/Private/EnvironmentQuery/EnvQueryInstanceBlueprintWrapper.cpp
--- a/Engine/Source/Runtime/AIModule/Private/EnvironmentQuery/EnvQueryInstanceBlueprintWrapper.cpp
+++ b/Engine/Source/Runtime/AIModule/Private/EnvironmentQuery/EnvQueryInstanceBlueprintWrapper.cpp
@@ -3,6 +3,7 @@
 #include "AIModulePrivate.h"
 #include "EnvironmentQuery/EnvQueryInstanceBlueprintWrapper.h"
 #include "EnvironmentQuery/Items/EnvQueryItemType_ActorBase.h"
+#include "EnvironmentQuery/EnvQueryManager.h"
 
 UEnvQueryInstanceBlueprintWrapper::UEnvQueryInstanceBlueprintWrapper(const FObjectInitializer& ObjectInitializer)
 	: Super(ObjectInitializer)
@@ -19,6 +20,13 @@ void UEnvQueryInstanceBlueprintWrapper::OnQueryFinished(TSharedPtr<FEnvQueryResu
 	OptionIndex = Result->OptionIndex;
 
 	OnQueryFinishedEvent.Broadcast(this, Result->GetRawStatus());
+
+	// unregister self, no longer shielded from GC
+	UEnvQueryManager* EnvQueryManager = Cast<UEnvQueryManager>(GetOuter());
+	if (ensure(EnvQueryManager))
+	{
+		EnvQueryManager->UnregisterActiveWrapper(*this);
+	}
 }
 
 float UEnvQueryInstanceBlueprintWrapper::GetItemScore(int32 ItemIndex)
@@ -62,4 +70,19 @@ TArray<FVector> UEnvQueryInstanceBlueprintWrapper::GetResultsAsLocations()
 	}
 
 	return Results;
+}
+
+void UEnvQueryInstanceBlueprintWrapper::RunQuery(const EEnvQueryRunMode::Type InRunMode, FEnvQueryRequest& QueryRequest)
+{
+	RunMode = InRunMode;
+	QueryID = QueryRequest.Execute(RunMode, this, &UEnvQueryInstanceBlueprintWrapper::OnQueryFinished);
+	if (QueryID != INDEX_NONE)
+	{
+		// register self as a wrapper needing shielding from GC
+		UEnvQueryManager* EnvQueryManager = Cast<UEnvQueryManager>(GetOuter());
+		if (ensure(EnvQueryManager))
+		{
+			EnvQueryManager->RegisterActiveWrapper(*this);
+		}
+	}
 }
\ No newline at end of file
diff --git a/Engine/Source/Runtime/AIModule/Private/EnvironmentQuery/EnvQueryManager.cpp b/Engine/Source/Runtime/AIModule/Private/EnvironmentQuery/EnvQueryManager.cpp
--- a/Engine/Source/Runtime/AIModule/Private/EnvironmentQuery/EnvQueryManager.cpp
+++ b/Engine/Source/Runtime/AIModule/Private/EnvironmentQuery/EnvQueryManager.cpp
@@ -267,7 +267,6 @@ void UEnvQueryManager::Tick(float DeltaTime)
 {
 	SCOPE_CYCLE_COUNTER(STAT_AI_EQS_Tick);
 	SET_DWORD_STAT(STAT_AI_EQS_NumInstances, RunningQueries.Num());
-	SCOPE_CYCLE_COUNTER(STAT_AI_Overall);
 	// @TODO: threads?
 
 	const double ExecutionTimeWarningSeconds = 0.25;
@@ -385,6 +384,8 @@ void UEnvQueryManager::OnWorldCleanup()
 			}
 		}
 	}
+
+	GCShieldedWrappers.Reset();
 }
 
 void UEnvQueryManager::RegisterExternalQuery(TSharedPtr<FEnvQueryInstance> QueryInstance)
@@ -694,7 +695,7 @@ float UEnvQueryManager::FindNamedParam(int32 QueryId, FName ParamName) const
 }
 
 //----------------------------------------------------------------------//
-// BP functions
+// BP functions and related functionality 
 //----------------------------------------------------------------------//
 UEnvQueryInstanceBlueprintWrapper* UEnvQueryManager::RunEQSQuery(UObject* WorldContext, UEnvQuery* QueryTemplate, UObject* Querier, TEnumAsByte<EEnvQueryRunMode::Type> RunMode, TSubclassOf<UEnvQueryInstanceBlueprintWrapper> WrapperClass)
 { 
@@ -727,21 +728,29 @@ UEnvQueryInstanceBlueprintWrapper* UEnvQueryManager::RunEQSQuery(UObject* WorldC
 
 		if (bValidQuerier)
 		{
-			QueryInstanceWrapper = NewObject<UEnvQueryInstanceBlueprintWrapper>((UClass*)(WrapperClass) ? (UClass*)WrapperClass : UEnvQueryInstanceBlueprintWrapper::StaticClass());
+			QueryInstanceWrapper = NewObject<UEnvQueryInstanceBlueprintWrapper>(EQSManager, (UClass*)(WrapperClass) ? (UClass*)WrapperClass : UEnvQueryInstanceBlueprintWrapper::StaticClass());
 			check(QueryInstanceWrapper);
 
 			FEnvQueryRequest QueryRequest(QueryTemplate, Querier);
 			// @todo named params still missing support
 			//QueryRequest.SetNamedParams(QueryParams);
-
-			QueryInstanceWrapper->SetRunMode(RunMode);
-			QueryInstanceWrapper->SetQueryID(QueryRequest.Execute(RunMode, QueryInstanceWrapper, &UEnvQueryInstanceBlueprintWrapper::OnQueryFinished));
+			QueryInstanceWrapper->RunQuery(RunMode, QueryRequest);
 		}
 	}
 	
 	return QueryInstanceWrapper;
 }
 
+void UEnvQueryManager::RegisterActiveWrapper(UEnvQueryInstanceBlueprintWrapper& Wrapper)
+{
+	GCShieldedWrappers.AddUnique(&Wrapper);
+}
+
+void UEnvQueryManager::UnregisterActiveWrapper(UEnvQueryInstanceBlueprintWrapper& Wrapper)
+{
+	GCShieldedWrappers.RemoveSingleSwap(&Wrapper, /*bAllowShrinking=*/false);
+}
+
 //----------------------------------------------------------------------//
 // Exec functions (i.e. console commands)
 //----------------------------------------------------------------------//
diff --git a/Engine/Source/Runtime/Core/Private/Misc/CoreMisc.cpp b/Engine/Source/Runtime/Core/Private/Misc/CoreMisc.cpp
--- a/Engine/Source/Runtime/Core/Private/Misc/CoreMisc.cpp
+++ b/Engine/Source/Runtime/Core/Private/Misc/CoreMisc.cpp
@@ -14,6 +14,7 @@
 #include "ModuleManager.h"
 #include "Ticker.h"
 #include "DerivedDataCacheInterface.h"
+#include "ITargetPlatformManagerModule.h"
 
 DEFINE_LOG_CATEGORY(LogSHA);
 DEFINE_LOG_CATEGORY(LogStats);
diff --git a/Engine/Source/Runtime/Core/Public/Modules/ModuleManager.h b/Engine/Source/Runtime/Core/Public/Modules/ModuleManager.h
--- a/Engine/Source/Runtime/Core/Public/Modules/ModuleManager.h
+++ b/Engine/Source/Runtime/Core/Public/Modules/ModuleManager.h
@@ -258,7 +258,7 @@ public:
 			return nullptr;
 		}
 
-		return (TModuleInterface*)(ModuleManager.GetModule(ModuleName).Get());
+		return static_cast<TModuleInterface*>(ModuleManager.GetModule(ModuleName).Get());
 	}
 
 	/**
diff --git a/Engine/Source/Runtime/Core/Public/ProfilingDebugging/DDCStatsHelper.h b/Engine/Source/Runtime/Core/Public/ProfilingDebugging/DDCStatsHelper.h
--- a/Engine/Source/Runtime/Core/Public/ProfilingDebugging/DDCStatsHelper.h
+++ b/Engine/Source/Runtime/Core/Public/ProfilingDebugging/DDCStatsHelper.h
@@ -1,5 +1,6 @@
-#pragma once
+// Copyright 1998-2015 Epic Games, Inc. All Rights Reserved.
 
+#pragma once
 
 class FDDCScopeStatHelper
 {
diff --git a/Engine/Source/Runtime/Engine/Private/AI/Navigation/RecastNavMesh.cpp b/Engine/Source/Runtime/Engine/Private/AI/Navigation/RecastNavMesh.cpp
--- a/Engine/Source/Runtime/Engine/Private/AI/Navigation/RecastNavMesh.cpp
+++ b/Engine/Source/Runtime/Engine/Private/AI/Navigation/RecastNavMesh.cpp
@@ -2070,6 +2070,7 @@ void ARecastNavMesh::PostEditChangeProperty(FPropertyChangedEvent& PropertyChang
 {
 	static const FName NAME_Generation = FName(TEXT("Generation"));
 	static const FName NAME_Display = FName(TEXT("Display"));
+	static const FName NAME_RuntimeGeneration = FName(TEXT("RuntimeGeneration"));
 	static const FName NAME_TileNumberHardLimit = GET_MEMBER_NAME_CHECKED(ARecastNavMesh, TileNumberHardLimit);
 
 	Super::PostEditChangeProperty(PropertyChangedEvent);
@@ -2112,6 +2113,18 @@ void ARecastNavMesh::PostEditChangeProperty(FPropertyChangedEvent& PropertyChang
 		{
 			RequestDrawingUpdate();
 		}
+		else if (PropertyChangedEvent.Property->GetFName() == NAME_RuntimeGeneration)
+		{
+			// @todo this contraption is required to clear RuntimeGeneration value in DefaultEngine.ini
+			// if it gets set to its default value (UE-23762). This is hopefully a temporary solution
+			// since it's an Core-level issue (UE-23873).
+			if (RuntimeGeneration == ERuntimeGenerationType::Static)
+			{
+				const FString EngineIniFilename = FPaths::ConvertRelativePathToFull(GetDefault<UEngine>()->GetDefaultConfigFilename());
+				GConfig->SetString(TEXT("/Script/Engine.RecastNavMesh"), *NAME_RuntimeGeneration.ToString(), TEXT(""), *EngineIniFilename);
+				GConfig->Flush(false);
+			}
+		}
 	}
 }
 
diff --git a/Engine/Source/Runtime/Engine/Private/Camera/CameraActor.cpp b/Engine/Source/Runtime/Engine/Private/Camera/CameraActor.cpp
--- a/Engine/Source/Runtime/Engine/Private/Camera/CameraActor.cpp
+++ b/Engine/Source/Runtime/Engine/Private/Camera/CameraActor.cpp
@@ -55,25 +55,24 @@ void ACameraActor::Serialize(FArchive& Ar)
 
 void ACameraActor::PostLoadSubobjects(FObjectInstancingGraph* OuterInstanceGraph)
 {
+	USceneComponent* OldRoot = RootComponent;
+	USceneComponent* OldAttachParent = OldRoot->AttachParent;
+	const FName OldSocketName = OldRoot->AttachSocketName;
+
+	Super::PostLoadSubobjects(OuterInstanceGraph);
+	
 	if (GetLinkerUE4Version() < VER_UE4_CAMERA_ACTOR_USING_CAMERA_COMPONENT)
 	{
-		USceneComponent* OldRoot = RootComponent;
-		USceneComponent* OldAttachParent = OldRoot->AttachParent;
-
-		Super::PostLoadSubobjects(OuterInstanceGraph);
-
 		CameraComponent->AttachParent = OldAttachParent;
 		OldRoot->AttachParent = NULL;
 	}
-	else
-	{
-		Super::PostLoadSubobjects(OuterInstanceGraph);
-	}
 
 	if (GetLinkerUE4Version() < VER_UE4_CAMERA_COMPONENT_ATTACH_TO_ROOT)
 	{
 		RootComponent = SceneComponent;
-		CameraComponent->AttachTo(RootComponent);
+		CameraComponent->AttachParent = RootComponent;
+		RootComponent->AttachParent = OldAttachParent;
+		RootComponent->AttachSocketName = OldSocketName;
 	}
 }
 
diff --git a/Engine/Source/Runtime/Engine/Private/UserInterface/PlayerInput.cpp b/Engine/Source/Runtime/Engine/Private/UserInterface/PlayerInput.cpp
--- a/Engine/Source/Runtime/Engine/Private/UserInterface/PlayerInput.cpp
+++ b/Engine/Source/Runtime/Engine/Private/UserInterface/PlayerInput.cpp
@@ -1042,12 +1042,14 @@ void UPlayerInput::ProcessInputStack(const TArray<UInputComponent*>& InputCompon
 						if (TB.bExecuteWhenPaused || !bGamePaused)
 						{
 							check(EventIndices.Num() > 0);
-							FDelegateDispatchDetails TouchInfo(EventIndices[0], NonAxisDelegates.Num(), TB.TouchDelegate, TouchEventLocations[EventIndices[0]], TouchIndex);
+							FVector *TouchedLocation = TouchEventLocations.Find(EventIndices[0]);
+							FDelegateDispatchDetails TouchInfo(EventIndices[0], NonAxisDelegates.Num(), TB.TouchDelegate, TouchedLocation != nullptr ? *TouchedLocation : FVector(-1.0f, -1.0f, 0.0f), TouchIndex);
 							NonAxisDelegates.Add(TouchInfo);
 							for (int32 EventsIndex = 1; EventsIndex < EventIndices.Num(); ++EventsIndex)
 							{
 								TouchInfo.EventIndex = EventIndices[EventsIndex];
-								TouchInfo.TouchLocation = TouchEventLocations[TouchInfo.EventIndex];
+								TouchedLocation = TouchEventLocations.Find(TouchInfo.EventIndex);
+								TouchInfo.TouchLocation = TouchedLocation != nullptr ? *TouchedLocation : FVector(-1.0f, -1.0f, 0.0f);
 								NonAxisDelegates.Add(TouchInfo);
 							}
 						}
diff --git a/Engine/Source/Runtime/MovieSceneCapture/Private/VideoCaptureProtocol.cpp b/Engine/Source/Runtime/MovieSceneCapture/Private/VideoCaptureProtocol.cpp
--- a/Engine/Source/Runtime/MovieSceneCapture/Private/VideoCaptureProtocol.cpp
+++ b/Engine/Source/Runtime/MovieSceneCapture/Private/VideoCaptureProtocol.cpp
@@ -29,7 +29,7 @@ bool FVideoCaptureProtocol::Initialize(const FCaptureProtocolInitSettings& InSet
 
 	if (CaptureSettings->bUseCompression)
 	{
-		Options.CompressionQuality = CaptureSettings->CompressionQuality;
+		Options.CompressionQuality = CaptureSettings->CompressionQuality / 100.f;
 		
 		float QualityOverride = 100.f;
 		if (FParse::Value( FCommandLine::Get(), TEXT( "-MovieQuality=" ), QualityOverride ))
diff --git a/Engine/Source/Runtime/Navmesh/Private/Recast/RecastLayers.cpp b/Engine/Source/Runtime/Navmesh/Private/Recast/RecastLayers.cpp
--- a/Engine/Source/Runtime/Navmesh/Private/Recast/RecastLayers.cpp
+++ b/Engine/Source/Runtime/Navmesh/Private/Recast/RecastLayers.cpp
@@ -1049,7 +1049,7 @@ bool rcBuildHeightfieldLayers(rcContext* ctx, rcCompactHeightfield& chf,
 	const int w = chf.width;
 	const int h = chf.height;
 	const int nreg = chf.maxRegions + 1;
-	rcLayerRegion* regions = (rcLayerRegion*)rcAlloc(sizeof(rcLayerRegion)*nreg, RC_ALLOC_TEMP);
+	rcScopedStructArrayDelete<rcLayerRegion> regions(nreg);
 	if (!regions)
 	{
 		ctx->log(RC_LOG_ERROR, "rcBuildHeightfieldLayers: Out of memory 'regions' (%d).", nreg);
diff --git a/Engine/Source/Runtime/Navmesh/Public/Recast/RecastAlloc.h b/Engine/Source/Runtime/Navmesh/Public/Recast/RecastAlloc.h
--- a/Engine/Source/Runtime/Navmesh/Public/Recast/RecastAlloc.h
+++ b/Engine/Source/Runtime/Navmesh/Public/Recast/RecastAlloc.h
@@ -129,4 +129,35 @@ public:
 	inline operator T*() { return ptr; }
 };
 
+/// A simple helper class used to delete an array of instances of structs,
+///	that require cleaning up by calling destructor on every instance before release. 
+/// Releasing is taking place when rcScopedStructArrayDelete instance goes out of scope.
+/// @note if T doesn't require a destructor to be called use rcScopedDelete instead.
+template<class T>
+class rcScopedStructArrayDelete
+{
+	int itemCount;
+	T* ptr;
+
+	// on purpose made private to restrict copying
+	inline T* operator=(T* p);
+public:
+
+	/// Constructs an array of instances of T
+	inline rcScopedStructArrayDelete(const int n) : itemCount(n) { ptr = (T*)rcAlloc(sizeof(T)*n, RC_ALLOC_TEMP); }
+
+	~rcScopedStructArrayDelete()
+	{
+		for (int itemIndex = 0; itemIndex < itemCount; ++itemIndex)
+		{
+			ptr[itemIndex].~T();
+		}
+		rcFree(ptr);
+	}
+
+	/// The root array pointer.
+	///  @return The root array pointer.
+	inline operator T*() { return ptr; }
+};
+
 #endif
diff --git a/Engine/Source/Runtime/Slate/Private/Widgets/Input/SButton.cpp b/Engine/Source/Runtime/Slate/Private/Widgets/Input/SButton.cpp
--- a/Engine/Source/Runtime/Slate/Private/Widgets/Input/SButton.cpp
+++ b/Engine/Source/Runtime/Slate/Private/Widgets/Input/SButton.cpp
@@ -73,7 +73,7 @@ int32 SButton::OnPaint(const FPaintArgs& Args, const FGeometry& AllottedGeometry
 	bool bEnabled = ShouldBeEnabled(bParentEnabled);
 	bool bShowDisabledEffect = GetShowDisabledEffect();
 
-	const FSlateBrush* BrushResource = GetBorder();
+	const FSlateBrush* BrushResource = !bShowDisabledEffect && !bEnabled ? DisabledImage : GetBorder();
 	
 	ESlateDrawEffect::Type DrawEffects = bShowDisabledEffect && !bEnabled ? ESlateDrawEffect::DisabledEffect : ESlateDrawEffect::None;
 
@@ -254,7 +254,12 @@ FReply SButton::OnMouseButtonUp( const FGeometry& MyGeometry, const FPointerEven
 		}
 		else
 		{
-			if( IsHovered() )
+			bool bEventOverButton = IsHovered();
+			if (!bEventOverButton && MouseEvent.IsTouchEvent())
+			{
+				bEventOverButton = MyGeometry.IsUnderLocation(MouseEvent.GetScreenSpacePosition());
+			}
+			if (bEventOverButton)
 			{
 				// If we asked for a precise tap, all we need is for the user to have not moved their pointer very far.
 				const bool bTriggerForTouchEvent = IsPreciseTapOrClick(MouseEvent);
diff --git a/Engine/Source/Runtime/UMG/Public/Slate/SlateVectorArtData.h b/Engine/Source/Runtime/UMG/Public/Slate/SlateVectorArtData.h
--- a/Engine/Source/Runtime/UMG/Public/Slate/SlateVectorArtData.h
+++ b/Engine/Source/Runtime/UMG/Public/Slate/SlateVectorArtData.h
@@ -1,4 +1,4 @@
-// Cop// Copyright 1998-2015 Epic Games, Inc. All Rights Reserved.
+// Copyright 1998-2015 Epic Games, Inc. All Rights Reserved.
 
 #pragma once
 
diff --git a/Engine/Source/Runtime/XmlParser/Public/FastXml.h b/Engine/Source/Runtime/XmlParser/Public/FastXml.h
--- a/Engine/Source/Runtime/XmlParser/Public/FastXml.h
+++ b/Engine/Source/Runtime/XmlParser/Public/FastXml.h
@@ -1,5 +1,6 @@
-#pragma once
+// Copyright 1998-2015 Epic Games, Inc. All Rights Reserved.
 
+#pragma once
 
 /**
  * Fast XML parser based on John W. Ratcliff's "FastXml" (see FastXml.tps and http://codesuppository.blogspot.com/2009/02/fastxml-extremely-lightweight-stream.html)
