Made the Plugin Warden entitlement check more reliable.

diff --git a/Engine/Source/Editor/PluginWarden/Private/PluginWardenAuthorizer.cpp b/Engine/Source/Editor/PluginWarden/Private/PluginWardenAuthorizer.cpp
--- a/Engine/Source/Editor/PluginWarden/Private/PluginWardenAuthorizer.cpp
+++ b/Engine/Source/Editor/PluginWarden/Private/PluginWardenAuthorizer.cpp
@@ -20,6 +20,16 @@ extern TSet<FString> AuthorizedPlugins;
 
 DEFINE_LOG_CATEGORY(PluginWarden);
 
+enum
+{
+	MaxAuthorizationRetries = 3,	// Max number of authorization check to retry when the entitlement wasn't retrieved yet
+	MaxSigninRetries = 3,			// Max number of prompt user for sign-in to retry when the Launcher fails to handle the prompt request
+	MaxLauncherRetries = 3,			// Max number of Launcher start to retry when it's detected that it's not available for user sign-in status check
+	GeneralWaitingTimeout = 15,		// Timeout in secs to use for various waiting phases
+	SigninWaitingTimeout = 120,		// Timeout in secs to use during the Waiting for sign-in phase
+	UserDetailsCheckPeriod = 10		// The periodicity in secs to check for the user details during the Waiting for sign-in phase
+};	
+
 #define LOCTEXT_NAMESPACE "PluginWardenAuthorization"
 
 FPluginWardenAuthorizer::FPluginWardenAuthorizer(const FText& InPluginFriendlyName, const FString& InPluginItemId, const FString& InPluginOfferId)
@@ -28,6 +38,10 @@ FPluginWardenAuthorizer::FPluginWardenAuthorizer(const FText& InPluginFriendlyNa
 	, PluginFriendlyName(InPluginFriendlyName)
 	, PluginItemId(InPluginItemId)
 	, PluginOfferId(InPluginOfferId)
+	, NumAuthorizationRetries(0)
+	, NumSignInRetries(0)
+	, NumLauncherRetries(0)
+	, CurrentWaitLoopNumber(0)
 {
 	TSharedRef<IPortalServiceLocator> ServiceLocator = GEditor->GetServiceLocator();
 	PortalWindowService = ServiceLocator->GetServiceRef<IPortalApplicationWindow>();
@@ -144,8 +158,26 @@ EPluginAuthorizationState FPluginWardenAuthorizer::UpdateAuthorizationState(floa
 
 				if ( UserDetails.IsSignedIn )
 				{
-					// if the user is signed in, and we're at this stage, we know they are unauthorized.
-					NewState = EPluginAuthorizationState::Unauthorized;
+					FPortalUserIsEntitledToItemResult Entitlement = EntitlementResult.GetFuture().Get();
+					if ( Entitlement.RetrievedFromCacheLevel == EEntitlementCacheLevelRetrieved::None )
+					{
+						// This is the case where there is no cached entitlement data and the latest entitlements have yet to be
+						// retrieved so we still don't know if the user is entitled so try again
+						if ( ++NumAuthorizationRetries < MaxAuthorizationRetries )
+						{
+							NewState = EPluginAuthorizationState::AuthorizePlugin;
+						}
+						else
+						{
+							// Give up and assume the user wasn't entitled
+							NewState = EPluginAuthorizationState::Unauthorized;
+						}
+					}
+					else
+					{
+						// The entitlement check was valid, the user is signed in and is not entitled, so clearly Unauthorized
+						NewState = EPluginAuthorizationState::Unauthorized;
+					}
 				}
 				else
 				{
@@ -155,6 +187,17 @@ EPluginAuthorizationState FPluginWardenAuthorizer::UpdateAuthorizationState(floa
 					{
 						NewState = EPluginAuthorizationState::SigninRequired;
 					}
+					else if ( ++NumLauncherRetries < MaxLauncherRetries )
+					{
+						// During testing, there's a flow where it goes through StartLauncher -> StartLauncher_Waiting -> AuthorizePlugin and
+						// ends up here without the Launcher actually running so try starting the Launcher again
+						NewState = EPluginAuthorizationState::StartLauncher;
+					}
+					else
+					{
+						// Give up and tell the user to sign in manually
+						NewState = EPluginAuthorizationState::SigninFailed;
+					}
 				}
 			}
 
@@ -174,24 +217,42 @@ EPluginAuthorizationState FPluginWardenAuthorizer::UpdateAuthorizationState(floa
 			check(UserSigninResult.GetFuture().IsValid());
 			if ( UserSigninResult.GetFuture().IsReady() )
 			{
-				bool IsUserSignedIn = UserSigninResult.GetFuture().Get();
-				if ( IsUserSignedIn )
+				// Note that the result of PromptUserForSignIn only says whether the portal successfully received and handled the request.
+				// It doesn't mean that the user signed in successfully or even that the existing user signed out.
+				bool IsUserPromptHandled = UserSigninResult.GetFuture().Get();
+				if ( IsUserPromptHandled )
 				{
+					// In this case, we only know that the user prompt was successful but we assume that the user signed in and that we can retrieve the details
 					UserDetailsResult = PortalUserService->GetUserDetails();
 					NewState = EPluginAuthorizationState::Signin_Waiting;
 					WaitingTime = 0;
 				}
 				else
 				{
-					NewState = EPluginAuthorizationState::Unauthorized;
+					// This state doesn't mean that the user is unauthorized but that the portal user login service wasn't able to handle the request
+					// So try again
+					if ( ++NumSignInRetries < MaxSigninRetries )
+					{
+						if (PortalUserLoginService->IsAvailable())
+						{
+							NewState = EPluginAuthorizationState::SigninRequired;
+						}
+						else
+						{
+							NewState = EPluginAuthorizationState::StartLauncher;
+						}
+					}
+					else
+					{
+						// Give up and tell the user to sign in manually
+						NewState = EPluginAuthorizationState::SigninFailed;
+					}
 				}
 			}
 
 			break;
 		}
-		// We stay in the Signin_Waiting state until the user is signed in or until they cancel the
-		// authorizing plug-in UI.  It would be nice to be able to know if the user closes the sign-in
-		// dialog and cancel out of this dialog automatically.
+		// We stay in the Signin_Waiting state until the user is signed in or the timeout is reached.
 		case EPluginAuthorizationState::Signin_Waiting:
 		{
 			WaitingTime += DeltaTime;
@@ -201,15 +262,17 @@ EPluginAuthorizationState FPluginWardenAuthorizer::UpdateAuthorizationState(floa
 			if ( UserDetailsResult.GetFuture().IsReady() )
 			{
 				FPortalUserDetails UserDetails = UserDetailsResult.GetFuture().Get();
+				int CurrentWaitSec = ((int)WaitingTime);
 
 				if ( UserDetails.IsSignedIn )
 				{
-					// if the user is signed in, and we're at this stage, we know they are unauthorized.
+					// If the user is now signed in, we can check for authorization again
 					NewState = EPluginAuthorizationState::AuthorizePlugin;
 				}
-				else
+				else if ( ( CurrentWaitSec % UserDetailsCheckPeriod ) == 0 && CurrentWaitSec != CurrentWaitLoopNumber )
 				{
-					NewState = EPluginAuthorizationState::SigninFailed;
+					// Every check period, try getting the UserDetails once to see if there's been any update in the signed-in status
+					UserDetailsResult = PortalUserService->GetUserDetails();
 				}
 			}
 
@@ -224,15 +287,21 @@ EPluginAuthorizationState FPluginWardenAuthorizer::UpdateAuthorizationState(floa
 		case EPluginAuthorizationState::AuthorizePlugin_Waiting:
 		case EPluginAuthorizationState::IsUserSignedIn_Waiting:
 		case EPluginAuthorizationState::SigninRequired_Waiting:
-		// We Ignore EPluginAuthorizationState::Signin_Waiting, that state could take forever, the user needs to sign-in or close the dialog.
 		{
-			const float TimeoutSeconds = 15;
-			if ( WaitingTime > TimeoutSeconds )
+			if ( WaitingTime > GeneralWaitingTimeout)
 			{
 				NewState = EPluginAuthorizationState::Timeout;
 			}
 			break;
 		}
+		case EPluginAuthorizationState::Signin_Waiting:
+		{
+			if ( WaitingTime > SigninWaitingTimeout )
+			{
+				NewState = EPluginAuthorizationState::SigninFailed;
+			}
+		}
+		break;
 	}
 
 	CurrentState = NewState;
@@ -240,4 +309,4 @@ EPluginAuthorizationState FPluginWardenAuthorizer::UpdateAuthorizationState(floa
 	return CurrentState;
 }
 
-#undef LOCTEXT_NAMESPACE
\ No newline at end of file
+#undef LOCTEXT_NAMESPACE
diff --git a/Engine/Source/Editor/PluginWarden/Private/PluginWardenAuthorizer.h b/Engine/Source/Editor/PluginWarden/Private/PluginWardenAuthorizer.h
--- a/Engine/Source/Editor/PluginWarden/Private/PluginWardenAuthorizer.h
+++ b/Engine/Source/Editor/PluginWarden/Private/PluginWardenAuthorizer.h
@@ -59,6 +59,18 @@ private:
 	/** The unique id of the offer for the plug-in on the marketplace. */
 	FString PluginOfferId;
 
+	/** The current number of authorization check retries. */
+	unsigned int NumAuthorizationRetries;
+
+	/** The current number of sign-in retries. */
+	unsigned int NumSignInRetries;
+
+	/** The current number of Launcher start retries. */
+	unsigned int NumLauncherRetries;
+
+	/** The current loop number of wait loops during the sign-in wait phase. */
+	unsigned int CurrentWaitLoopNumber;
+
 	/** The portal application communication service. */
 	TSharedPtr<IPortalApplicationWindow> PortalWindowService;
 
diff --git a/Engine/Source/Editor/PluginWarden/Private/SAuthorizingPlugin.cpp b/Engine/Source/Editor/PluginWarden/Private/SAuthorizingPlugin.cpp
--- a/Engine/Source/Editor/PluginWarden/Private/SAuthorizingPlugin.cpp
+++ b/Engine/Source/Editor/PluginWarden/Private/SAuthorizingPlugin.cpp
@@ -107,7 +107,7 @@ FText SAuthorizingPlugin::GetWaitingText() const
 	case EPluginAuthorizationState::SigninRequired_Waiting:
 		return LOCTEXT("NeedUserToLoginToCheck", "Authorization failed, sign-in required...");
 	case EPluginAuthorizationState::Signin_Waiting:
-		return LOCTEXT("WaitingForSignin", "Waiting for sign-in...");
+		return LOCTEXT("WaitingForSignin", "Epic Games Launcher started. Waiting for you to sign in...");
 	case EPluginAuthorizationState::SigninFailed:
 		return LOCTEXT("SigninFailed", "Sign-in failed. Cancel and sign in manually through the launcher.");
 	}
