Copied 4393315 JeanMichel.Dignard
Better telemetry for plugin warden
- Added more event attributes to track potential issues.

diff --git a/Engine/Source/Editor/PluginWarden/Private/PluginWardenModule.cpp b/Engine/Source/Editor/PluginWarden/Private/PluginWardenModule.cpp
--- a/Engine/Source/Editor/PluginWarden/Private/PluginWardenModule.cpp
+++ b/Engine/Source/Editor/PluginWarden/Private/PluginWardenModule.cpp
@@ -2,9 +2,11 @@
 
 #include "PluginWardenModule.h"
 #include "Async/TaskGraphInterfaces.h"
+#include "EngineAnalytics.h"
+#include "Framework/Application/SlateApplication.h"
+#include "Interfaces/IAnalyticsProvider.h"
 #include "Widgets/DeclarativeSyntaxSupport.h"
 #include "Widgets/SWindow.h"
-#include "Framework/Application/SlateApplication.h"
 
 #include "PluginWardenAuthorizer.h"
 #include "SAuthorizingPlugin.h"
@@ -66,6 +68,8 @@ bool FPluginWardenModule::RunAuthorizationPipeline(const FText& PluginFriendlyNa
 	FPluginWardenAuthorizer Authorizer(PluginFriendlyName, PluginItemId, PluginOfferId);
 
 	EPluginAuthorizationState AuthorizationState = EPluginAuthorizationState::Initializing;
+	EPluginAuthorizationState PreviousState = AuthorizationState;
+
 	bool bAuthorizationCompleted = false;
 
 	double LastLoopTime = FPlatformTime::Seconds();
@@ -91,7 +95,7 @@ bool FPluginWardenModule::RunAuthorizationPipeline(const FText& PluginFriendlyNa
 		FTaskGraphInterface::Get().ProcessThreadUntilIdle(ENamedThreads::GameThread);
 		FTicker::GetCoreTicker().Tick(DeltaTime);
 
-		const EPluginAuthorizationState PreviousState = AuthorizationState;
+		PreviousState = AuthorizationState;
 		AuthorizationState = Authorizer.UpdateAuthorizationState(DeltaTime);
 
 		switch ( AuthorizationState )
@@ -164,6 +168,15 @@ bool FPluginWardenModule::RunAuthorizationPipeline(const FText& PluginFriendlyNa
 		}
 	}
 
+	TArray< FAnalyticsEventAttribute > EventAttributes;
+	EventAttributes.Emplace( TEXT("State"), (int32)AuthorizationState );
+	EventAttributes.Emplace( TEXT("PreviousState"), (int32)PreviousState );
+	EventAttributes.Emplace( TEXT("UnauthorizedErrorHandling"), (int32)IPluginWardenModule::EUnauthorizedErrorHandling::Silent );
+	EventAttributes.Emplace( TEXT("ItemId"), PluginItemId );
+	EventAttributes.Emplace( TEXT("OfferId"), PluginOfferId );
+
+	FEngineAnalytics::GetProvider().RecordEvent( TEXT("PluginWarden.AuthorizationFailure"), EventAttributes );
+
 	return false;
 }
 
diff --git a/Engine/Source/Editor/PluginWarden/Private/SAuthorizingPlugin.cpp b/Engine/Source/Editor/PluginWarden/Private/SAuthorizingPlugin.cpp
--- a/Engine/Source/Editor/PluginWarden/Private/SAuthorizingPlugin.cpp
+++ b/Engine/Source/Editor/PluginWarden/Private/SAuthorizingPlugin.cpp
@@ -32,6 +32,9 @@ void SAuthorizingPlugin::Construct(const FArguments& InArgs, const TSharedRef<SW
 	InParentWindow->SetOnWindowClosed(FOnWindowClosed::CreateSP(this, &SAuthorizingPlugin::OnWindowClosed));
 	bUserInterrupted = true;
 
+	AuthorizationState = EPluginAuthorizationState::Initializing;
+	PreviousAuthorizationState = AuthorizationState;
+
 	RegisterActiveTimer(0.f, FWidgetActiveTimerDelegate::CreateSP(this, &SAuthorizingPlugin::RefreshStatus));
 
 	ChildSlot
@@ -122,6 +125,7 @@ EActiveTimerReturnType SAuthorizingPlugin::RefreshStatus(double InCurrentTime, f
 	FTaskGraphInterface::Get().ProcessThreadUntilIdle(ENamedThreads::GameThread);
 	FTicker::GetCoreTicker().Tick(InDeltaTime);
 
+	PreviousAuthorizationState = AuthorizationState;
 	AuthorizationState = Authorizer->UpdateAuthorizationState(InDeltaTime);
 
 	switch ( AuthorizationState )
@@ -219,7 +223,14 @@ void SAuthorizingPlugin::OnWindowClosed(const TSharedRef<SWindow>& InWindow)
 			}
 		}
 
-		FEngineAnalytics::GetProvider().RecordEvent( TEXT("PluginWarden.AuthorizationFailure"), FAnalyticsEventAttribute( TEXT("State"), (int32)AuthorizationState ) );
+		TArray< FAnalyticsEventAttribute > EventAttributes;
+		EventAttributes.Emplace( TEXT("State"), (int32)AuthorizationState );
+		EventAttributes.Emplace( TEXT("PreviousState"), (int32)PreviousAuthorizationState );
+		EventAttributes.Emplace( TEXT("UnauthorizedErrorHandling"), (int32)UnauthorizedErrorHandling );
+		EventAttributes.Emplace( TEXT("ItemId"), Authorizer->GetPluginItemId() );
+		EventAttributes.Emplace( TEXT("OfferId"), Authorizer->GetPluginOfferId() );
+
+		FEngineAnalytics::GetProvider().RecordEvent( TEXT("PluginWarden.AuthorizationFailure"), EventAttributes );
 	}
 }
 
diff --git a/Engine/Source/Editor/PluginWarden/Private/SAuthorizingPlugin.h b/Engine/Source/Editor/PluginWarden/Private/SAuthorizingPlugin.h
--- a/Engine/Source/Editor/PluginWarden/Private/SAuthorizingPlugin.h
+++ b/Engine/Source/Editor/PluginWarden/Private/SAuthorizingPlugin.h
@@ -64,6 +64,9 @@ private:
 	/** The latest state of the plug-in authorization pipeline. */
 	EPluginAuthorizationState AuthorizationState;
 
+	/** The previous state of the plug-in authorization pipeline. */
+	EPluginAuthorizationState PreviousAuthorizationState;
+
 	/** If the user is authorized to us the plug-in, we'll call this function to alert the plug-in that everything is good to go. */
 	TFunction<void()> AuthorizedCallback;
 
