Copied 4393716 JeanMichel.Dignard
Changed Datasmith & MDL entitlements to use the cache on disk (user not required to be logged in)

diff --git a/Engine/Source/Editor/PluginWarden/Private/PluginWardenAuthorizer.cpp b/Engine/Source/Editor/PluginWarden/Private/PluginWardenAuthorizer.cpp
--- a/Engine/Source/Editor/PluginWarden/Private/PluginWardenAuthorizer.cpp
+++ b/Engine/Source/Editor/PluginWarden/Private/PluginWardenAuthorizer.cpp
@@ -32,8 +32,9 @@ enum
 
 #define LOCTEXT_NAMESPACE "PluginWardenAuthorization"
 
-FPluginWardenAuthorizer::FPluginWardenAuthorizer(const FText& InPluginFriendlyName, const FString& InPluginItemId, const FString& InPluginOfferId)
+FPluginWardenAuthorizer::FPluginWardenAuthorizer(const FText& InPluginFriendlyName, const FString& InPluginItemId, const FString& InPluginOfferId, const EEntitlementCacheLevelRequest InCacheLevel)
 	: CurrentState(EPluginAuthorizationState::Initializing)
+	, CacheLevel(InCacheLevel)
 	, WaitingTime(0)
 	, PluginFriendlyName(InPluginFriendlyName)
 	, PluginItemId(InPluginItemId)
@@ -116,7 +117,7 @@ EPluginAuthorizationState FPluginWardenAuthorizer::UpdateAuthorizationState(floa
 		case EPluginAuthorizationState::AuthorizePlugin:
 		{
 			WaitingTime = 0;
-			EntitlementResult = PortalUserService->IsEntitledToItem(PluginItemId, EEntitlementCacheLevelRequest::Memory);
+			EntitlementResult = PortalUserService->IsEntitledToItem(PluginItemId, CacheLevel);
 			NewState = EPluginAuthorizationState::AuthorizePlugin_Waiting;
 			break;
 		}
diff --git a/Engine/Source/Editor/PluginWarden/Private/PluginWardenAuthorizer.h b/Engine/Source/Editor/PluginWarden/Private/PluginWardenAuthorizer.h
--- a/Engine/Source/Editor/PluginWarden/Private/PluginWardenAuthorizer.h
+++ b/Engine/Source/Editor/PluginWarden/Private/PluginWardenAuthorizer.h
@@ -34,7 +34,7 @@ enum class EPluginAuthorizationState
 class FPluginWardenAuthorizer
 {
 public:
-	FPluginWardenAuthorizer(const FText& InPluginFriendlyName, const FString& InPluginItemId, const FString& InPluginOfferId);
+	FPluginWardenAuthorizer(const FText& InPluginFriendlyName, const FString& InPluginItemId, const FString& InPluginOfferId, const EEntitlementCacheLevelRequest InCacheLevel);
 
 	/** Returns new state after elapsed time */
 	EPluginAuthorizationState UpdateAuthorizationState(float DeltaTime);
@@ -47,6 +47,9 @@ private:
 	/** The current state of the plug-in authorization pipeline. */
 	EPluginAuthorizationState CurrentState;
 
+	/** Cache level to check for the entitlement */
+	EEntitlementCacheLevelRequest CacheLevel;
+
 	/** The amount of time we've been waiting for confirmation for a given step.  It's possible a problem may arise and we need to timeout. */
 	float WaitingTime;
 
diff --git a/Engine/Source/Editor/PluginWarden/Private/PluginWardenModule.cpp b/Engine/Source/Editor/PluginWarden/Private/PluginWardenModule.cpp
--- a/Engine/Source/Editor/PluginWarden/Private/PluginWardenModule.cpp
+++ b/Engine/Source/Editor/PluginWarden/Private/PluginWardenModule.cpp
@@ -27,7 +27,8 @@ void FPluginWardenModule::ShutdownModule()
 
 }
 
-void FPluginWardenModule::CheckEntitlementForPlugin(const FText& PluginFriendlyName, const FString& PluginItemId, const FString& PluginOfferId, const FText& UnauthorizedMessageOverride, EUnauthorizedErrorHandling UnauthorizedErrorHandling, TFunction<void()> AuthorizedCallback)
+void FPluginWardenModule::CheckEntitlementForPlugin(const FText& PluginFriendlyName, const FString& PluginItemId, const FString& PluginOfferId, const EEntitlementCacheLevelRequest CacheLevel,
+	const FText& UnauthorizedMessageOverride, EUnauthorizedErrorHandling UnauthorizedErrorHandling, TFunction<void()> AuthorizedCallback)
 {
 	// If we've previously authorized the plug-in, just immediately verify access.
 	if ( AuthorizedPlugins.Contains(PluginItemId) )
@@ -38,7 +39,7 @@ void FPluginWardenModule::CheckEntitlementForPlugin(const FText& PluginFriendlyN
 
 	if (IsRunningCommandlet() || GIsRunningUnattendedScript)
 	{
-		if (RunAuthorizationPipeline(PluginFriendlyName, PluginItemId, PluginOfferId))
+		if (RunAuthorizationPipeline(PluginFriendlyName, PluginItemId, PluginOfferId, CacheLevel))
 		{
 			AuthorizedPlugins.Add(PluginItemId);
 			AuthorizedCallback();
@@ -54,7 +55,7 @@ void FPluginWardenModule::CheckEntitlementForPlugin(const FText& PluginFriendlyN
 			.SizingRule(ESizingRule::Autosized)
 			.Title(FText::Format(LOCTEXT("EntitlementCheckFormat", "{0} - Entitlement Check"), PluginFriendlyName));
 
-		TSharedRef<SAuthorizingPlugin> PluginAuthPanel = SNew(SAuthorizingPlugin, AuthorizingPluginWindow, PluginFriendlyName, PluginItemId, PluginOfferId, AuthorizedCallback);
+		TSharedRef<SAuthorizingPlugin> PluginAuthPanel = SNew(SAuthorizingPlugin, AuthorizingPluginWindow, SAuthorizingPlugin::FPluginInfo( PluginFriendlyName, PluginItemId, PluginOfferId ), CacheLevel, AuthorizedCallback);
 		PluginAuthPanel->SetUnauthorizedOverride(UnauthorizedMessageOverride, UnauthorizedErrorHandling);
 
 		AuthorizingPluginWindow->SetContent(PluginAuthPanel);
@@ -63,9 +64,14 @@ void FPluginWardenModule::CheckEntitlementForPlugin(const FText& PluginFriendlyN
 	}
 }
 
-bool FPluginWardenModule::RunAuthorizationPipeline(const FText& PluginFriendlyName, const FString& PluginItemId, const FString& PluginOfferId)
+void FPluginWardenModule::CheckEntitlementForPlugin(const FText& PluginFriendlyName, const FString& PluginItemId, const FString& PluginOfferId, const FText& UnauthorizedMessageOverride, EUnauthorizedErrorHandling UnauthorizedErrorHandling, TFunction<void()> AuthorizedCallback)
+{
+	CheckEntitlementForPlugin(PluginFriendlyName, PluginItemId, PluginOfferId, EEntitlementCacheLevelRequest::Memory, UnauthorizedMessageOverride, UnauthorizedErrorHandling, AuthorizedCallback);
+}
+
+bool FPluginWardenModule::RunAuthorizationPipeline(const FText& PluginFriendlyName, const FString& PluginItemId, const FString& PluginOfferId, const EEntitlementCacheLevelRequest CacheLevel)
 {
-	FPluginWardenAuthorizer Authorizer(PluginFriendlyName, PluginItemId, PluginOfferId);
+	FPluginWardenAuthorizer Authorizer(PluginFriendlyName, PluginItemId, PluginOfferId, CacheLevel);
 
 	EPluginAuthorizationState AuthorizationState = EPluginAuthorizationState::Initializing;
 	EPluginAuthorizationState PreviousState = AuthorizationState;
diff --git a/Engine/Source/Editor/PluginWarden/Private/PluginWardenModule.h b/Engine/Source/Editor/PluginWarden/Private/PluginWardenModule.h
--- a/Engine/Source/Editor/PluginWarden/Private/PluginWardenModule.h
+++ b/Engine/Source/Editor/PluginWarden/Private/PluginWardenModule.h
@@ -38,12 +38,17 @@ public:
 	 * @param PluginFriendlyName The localized friendly name of the plug-in.
 	 * @param PluginItemId The unique identifier of the item plug-in on the marketplace.
 	 * @param PluginOfferId The unique identifier of the offer for the plug-in on the marketplace.
+	 * @param CacheLevel Where to check for the entitlements. Defaults to memory.
 	 * @param UnauthorizedMessageOverride The error message to display for unauthorized plugins, overriding the default message if not empty.
 	 * @param UnauthorizedErrorHandling How to handle the unauthorized error.
 	 * @param AuthorizedCallback This function will be called after the user has been given entitlement.
 	 */
-	virtual void CheckEntitlementForPlugin(const FText& PluginFriendlyName, const FString& PluginItemId, const FString& PluginOfferId, const FText& UnauthorizedMessageOverride, EUnauthorizedErrorHandling UnauthorizedErrorHandling, TFunction<void()> AuthorizedCallback) override;
+	virtual void CheckEntitlementForPlugin(const FText& PluginFriendlyName, const FString& PluginItemId, const FString& PluginOfferId, const EEntitlementCacheLevelRequest CacheLevel,
+		const FText& UnauthorizedMessageOverride, EUnauthorizedErrorHandling UnauthorizedErrorHandling, TFunction<void()> AuthorizedCallback) override;
+
+	virtual void CheckEntitlementForPlugin(const FText& PluginFriendlyName, const FString& PluginItemId, const FString& PluginOfferId,
+		const FText& UnauthorizedMessageOverride, EUnauthorizedErrorHandling UnauthorizedErrorHandling, TFunction<void()> AuthorizedCallback) override;
 
 private:
-	bool RunAuthorizationPipeline(const FText& PluginFriendlyName, const FString& PluginItemId, const FString& PluginOfferId);
+	bool RunAuthorizationPipeline(const FText& PluginFriendlyName, const FString& PluginItemId, const FString& PluginOfferId, const EEntitlementCacheLevelRequest CacheLevel);
 };
diff --git a/Engine/Source/Editor/PluginWarden/Private/SAuthorizingPlugin.cpp b/Engine/Source/Editor/PluginWarden/Private/SAuthorizingPlugin.cpp
--- a/Engine/Source/Editor/PluginWarden/Private/SAuthorizingPlugin.cpp
+++ b/Engine/Source/Editor/PluginWarden/Private/SAuthorizingPlugin.cpp
@@ -21,13 +21,13 @@
 
 #define LOCTEXT_NAMESPACE "PluginWarden"
 
-void SAuthorizingPlugin::Construct(const FArguments& InArgs, const TSharedRef<SWindow>& InParentWindow, const FText& InPluginFriendlyName, const FString& InPluginItemId, const FString& InPluginOfferId, TFunction<void()> InAuthorizedCallback)
+void SAuthorizingPlugin::Construct(const FArguments& InArgs, const TSharedRef<SWindow>& InParentWindow, const FPluginInfo& InPluginInfo, const EEntitlementCacheLevelRequest InCacheLevel, TFunction<void()> InAuthorizedCallback)
 {
 	ParentWindow = InParentWindow;
 	AuthorizedCallback = InAuthorizedCallback;
 	UnauthorizedErrorHandling = IPluginWardenModule::EUnauthorizedErrorHandling::ShowMessageOpenStore;
 
-	Authorizer = MakeShared<FPluginWardenAuthorizer>(InPluginFriendlyName, InPluginItemId, InPluginOfferId);
+	Authorizer = MakeShared<FPluginWardenAuthorizer>(InPluginInfo.FriendlyName, InPluginInfo.ItemId, InPluginInfo.OfferId, InCacheLevel);
 
 	InParentWindow->SetOnWindowClosed(FOnWindowClosed::CreateSP(this, &SAuthorizingPlugin::OnWindowClosed));
 	bUserInterrupted = true;
diff --git a/Engine/Source/Editor/PluginWarden/Private/SAuthorizingPlugin.h b/Engine/Source/Editor/PluginWarden/Private/SAuthorizingPlugin.h
--- a/Engine/Source/Editor/PluginWarden/Private/SAuthorizingPlugin.h
+++ b/Engine/Source/Editor/PluginWarden/Private/SAuthorizingPlugin.h
@@ -10,6 +10,7 @@
 #include "IPluginWardenModule.h"
 #include "PluginWardenAuthorizer.h"
 
+enum class EEntitlementCacheLevelRequest : uint8;
 class FPluginWardenAuthorizer;
 
 extern TSet<FString> AuthorizedPlugins;
@@ -24,7 +25,21 @@ public:
 
 	SLATE_END_ARGS()
 
-	void Construct(const FArguments& InArgs, const TSharedRef<SWindow>& InParentWindow, const FText& InPluginFriendlyName, const FString& InPluginItemId, const FString& InPluginOfferId, TFunction<void()> InAuthorizedCallback);
+	struct FPluginInfo
+	{
+		FPluginInfo( const FText& InFriendlyName, const FString& InItemId, const FString& InOfferId )
+			: FriendlyName( InFriendlyName )
+			, ItemId( InItemId )
+			, OfferId( InOfferId )
+		{
+		}
+
+		const FText& FriendlyName;
+		const FString& ItemId;
+		const FString& OfferId;
+	};
+
+	void Construct(const FArguments& InArgs, const TSharedRef<SWindow>& InParentWindow, const FPluginInfo& PluginInfo, const EEntitlementCacheLevelRequest CacheLevel, TFunction<void()> InAuthorizedCallback);
 
 	/**
 	 * Override the default message and behavior in the case where the plugin is unauthorized.
diff --git a/Engine/Source/Editor/PluginWarden/Public/IPluginWardenModule.h b/Engine/Source/Editor/PluginWarden/Public/IPluginWardenModule.h
--- a/Engine/Source/Editor/PluginWarden/Public/IPluginWardenModule.h
+++ b/Engine/Source/Editor/PluginWarden/Public/IPluginWardenModule.h
@@ -5,6 +5,8 @@
 #include "Modules/ModuleInterface.h"
 #include "Modules/ModuleManager.h"
 
+enum class EEntitlementCacheLevelRequest : uint8;
+
 /**
  * The Plugin Warden is a simple module used to verify a user has purchased a plug-in.  This
  * module won't prevent a determined user from avoiding paying for a plug-in, it is merely to
@@ -69,9 +71,14 @@ public:
 	 * @param PluginFriendlyName The localized friendly name of the plug-in.
 	 * @param PluginItemId The unique identifier of the item plug-in on the marketplace.
 	 * @param PluginOfferId The unique identifier of the offer for the plug-in on the marketplace.
+	 * @param CacheLevel Where to check for the entitlements. Defaults to memory.
 	 * @param UnauthorizedMessageOverride The error message to display for unauthorized plugins, overriding the default message if not empty.
 	 * @param UnauthorizedErrorHandling How to handle the unauthorized error.
 	 * @param AuthorizedCallback This function will be called after the user has been given entitlement.
 	 */
-	virtual void CheckEntitlementForPlugin(const FText& PluginFriendlyName, const FString& PluginItemId, const FString& PluginOfferId, const FText& UnauthorizedMessageOverride, EUnauthorizedErrorHandling UnauthorizedErrorHandling, TFunction<void()> AuthorizedCallback) = 0;
+	virtual void CheckEntitlementForPlugin(const FText& PluginFriendlyName, const FString& PluginItemId, const FString& PluginOfferId, const EEntitlementCacheLevelRequest CacheLevel,
+		const FText& UnauthorizedMessageOverride, EUnauthorizedErrorHandling UnauthorizedErrorHandling, TFunction<void()> AuthorizedCallback) = 0;
+
+	virtual void CheckEntitlementForPlugin(const FText& PluginFriendlyName, const FString& PluginItemId, const FString& PluginOfferId,
+		const FText& UnauthorizedMessageOverride, EUnauthorizedErrorHandling UnauthorizedErrorHandling, TFunction<void()> AuthorizedCallback) = 0;
 };
